Update VirtualHeightfieldMesh to use unified RDGGraph
#rnx

diff --git a/Engine/Plugins/Experimental/VirtualHeightfieldMesh/Source/VirtualHeightfieldMesh/Private/VirtualHeightfieldMeshSceneProxy.cpp b/Engine/Plugins/Experimental/VirtualHeightfieldMesh/Source/VirtualHeightfieldMesh/Private/VirtualHeightfieldMeshSceneProxy.cpp
--- a/Engine/Plugins/Experimental/VirtualHeightfieldMesh/Source/VirtualHeightfieldMesh/Private/VirtualHeightfieldMeshSceneProxy.cpp
+++ b/Engine/Plugins/Experimental/VirtualHeightfieldMesh/Source/VirtualHeightfieldMesh/Private/VirtualHeightfieldMeshSceneProxy.cpp
@@ -83,11 +83,13 @@ namespace VirtualHeightfieldMesh
 	struct FDrawInstanceBuffers
 	{
 		/* Culled instance buffer. */
-		TRefCountPtr<FRDGPooledBuffer> InstanceBuffer;
+		FStructuredBufferRHIRef InstanceBuffer;
+		FUnorderedAccessViewRHIRef InstanceBufferUAV;
 		FShaderResourceViewRHIRef InstanceBufferSRV;
 
 		/* IndirectArgs buffer for final DrawInstancedIndirect. */
-		TRefCountPtr<FRDGPooledBuffer> IndirectArgsBuffer;
+		FVertexBufferRHIRef IndirectArgsBuffer;
+		FUnorderedAccessViewRHIRef IndirectArgsBufferUAV;
 	};
 
 	/** Initialize the FDrawInstanceBuffers objects. */
@@ -97,8 +99,10 @@ namespace VirtualHeightfieldMesh
 	void ReleaseInstanceBuffers(FDrawInstanceBuffers& InBuffers)
 	{
 		InBuffers.InstanceBuffer.SafeRelease();
+		InBuffers.InstanceBufferUAV.SafeRelease();
 		InBuffers.InstanceBufferSRV.SafeRelease();
 		InBuffers.IndirectArgsBuffer.SafeRelease();
+		InBuffers.IndirectArgsBufferUAV.SafeRelease();
 	}
 }
 
@@ -170,7 +174,7 @@ public:
 	/** Call once per frame for each mesh/view that has relevance. This allocates the buffers to use for the frame and adds the work to fill the buffers to the queue. */
 	VirtualHeightfieldMesh::FDrawInstanceBuffers& AddWork(FVirtualHeightfieldMeshSceneProxy const* InProxy, FSceneView const* InMainView, FSceneView const* InCullView);
 	/** Submit all the work added by AddWork(). The work fills all of the buffers ready for use by the referencing mesh batches. */
-	void SubmitWork(FRHICommandListImmediate& InRHICmdList);
+	void SubmitWork(FRDGBuilder& GraphBuilder);
 
 protected:
 	//~ Begin FRenderResource Interface
@@ -179,8 +183,9 @@ protected:
 
 private:
 	/** Called by renderer at start of render frame. */
-	void BeginFrame();
+	void BeginFrame(FRDGBuilder& GraphBuilder);
 	/** Called by renderer at end of render frame. */
+	void EndFrame(FRDGBuilder& GraphBuilder);
 	void EndFrame();
 
 	/** Flag for frame validation. */
@@ -235,8 +240,8 @@ void FVirtualHeightfieldMeshRendererExtension::RegisterExtension()
 	static bool bInit = false;
 	if (!bInit)
 	{
-		GEngine->GetPreRenderDelegate().AddRaw(this, &FVirtualHeightfieldMeshRendererExtension::BeginFrame);
-		GEngine->GetPostRenderDelegate().AddRaw(this, &FVirtualHeightfieldMeshRendererExtension::EndFrame);
+		GEngine->GetPreRenderDelegateEx().AddRaw(this, &FVirtualHeightfieldMeshRendererExtension::BeginFrame);
+		GEngine->GetPostRenderDelegateEx().AddRaw(this, &FVirtualHeightfieldMeshRendererExtension::EndFrame);
 		bInit = true;
 	}
 }
@@ -298,7 +303,7 @@ VirtualHeightfieldMesh::FDrawInstanceBuffers& FVirtualHeightfieldMeshRendererExt
 	return Buffers[WorkDesc.BufferIndex];
 }
 
-void FVirtualHeightfieldMeshRendererExtension::BeginFrame()
+void FVirtualHeightfieldMeshRendererExtension::BeginFrame(FRDGBuilder& GraphBuilder)
 {
 	// If we hit this then BegineFrame()/EndFrame() logic needs fixing in the Scene Renderer.
 	if (!ensure(!bInFrame))
@@ -309,7 +314,7 @@ void FVirtualHeightfieldMeshRendererExtension::BeginFrame()
 
 	if (WorkDescs.Num() > 0)
 	{
-		SubmitWork(GetImmediateCommandList_ForRenderCommand());
+		SubmitWork(GraphBuilder);
 	}
 }
 
@@ -343,6 +348,11 @@ void FVirtualHeightfieldMeshRendererExtension::EndFrame()
 	GOcclusionResults.Reset();
 }
 
+void FVirtualHeightfieldMeshRendererExtension::EndFrame(FRDGBuilder& GraphBuilder)
+{
+	EndFrame();
+}
+
 const static FName NAME_VirtualHeightfieldMesh(TEXT("VirtualHeightfieldMesh"));
 
 FVirtualHeightfieldMeshSceneProxy::FVirtualHeightfieldMeshSceneProxy(UVirtualHeightfieldMeshComponent* InComponent)
@@ -498,7 +508,7 @@ void FVirtualHeightfieldMeshSceneProxy::GetDynamicMeshElements(const TArray<cons
 				FMeshBatchElement& BatchElement = Mesh.Elements[0];
 
 				BatchElement.IndexBuffer = VertexFactory->IndexBuffer;
-				BatchElement.IndirectArgsBuffer = Buffers.IndirectArgsBuffer->GetRHI();
+				BatchElement.IndirectArgsBuffer = Buffers.IndirectArgsBuffer;
 				BatchElement.IndirectArgsOffset = 0;
 
 				BatchElement.FirstIndex = 0;
@@ -807,7 +817,7 @@ namespace VirtualHeightfieldMesh
 
 		BEGIN_SHADER_PARAMETER_STRUCT(FParameters, )
 			SHADER_PARAMETER(int32, NumIndices)
-			SHADER_PARAMETER_RDG_BUFFER_UAV(RWBuffer<uint>, RWIndirectArgsBuffer)
+			SHADER_PARAMETER_UAV(RWBuffer<uint>, RWIndirectArgsBuffer)
 		END_SHADER_PARAMETER_STRUCT()
 
 		static bool ShouldCompilePermutation(FGlobalShaderPermutationParameters const& Parameters)
@@ -836,8 +846,8 @@ namespace VirtualHeightfieldMesh
 			SHADER_PARAMETER_RDG_BUFFER_SRV(Buffer<uint>, QuadNeighborBuffer)
 			SHADER_PARAMETER_RDG_BUFFER(Buffer<uint>, IndirectArgsBuffer)
 			SHADER_PARAMETER_RDG_BUFFER_SRV(Buffer<uint>, IndirectArgsBufferSRV)
-			SHADER_PARAMETER_RDG_BUFFER_UAV(RWStructuredBuffer<QuadRenderInstance>, RWInstanceBuffer)
-			SHADER_PARAMETER_RDG_BUFFER_UAV(RWBuffer<uint>, RWIndirectArgsBuffer)
+			SHADER_PARAMETER_UAV(RWStructuredBuffer<QuadRenderInstance>, RWInstanceBuffer)
+			SHADER_PARAMETER_UAV(RWBuffer<uint>, RWIndirectArgsBuffer)
 		END_SHADER_PARAMETER_STRUCT()
 	};
 
@@ -1023,17 +1033,6 @@ namespace VirtualHeightfieldMesh
 		FRDGBufferSRVRef QuadNeighborBufferSRV;
 	};
 
-	/** Structure to carry the RDG wrapping for our output FDrawInstanceBuffers. */
-	struct FOutputResources
-	{
-		FRDGBufferRef InstanceBuffer;
-		FRDGBufferUAVRef InstanceBufferUAV;
-		FRDGBufferSRVRef InstanceBufferSRV;
-
-		FRDGBufferRef IndirectArgsBuffer;
-		FRDGBufferUAVRef IndirectArgsBufferUAV;
-	};
-
 	/* Dummy parameter struct used to allocate FPooledRDGBuffer objects using a fake RDG pass. */
 	BEGIN_SHADER_PARAMETER_STRUCT(FCreateBufferParameters, )
 		SHADER_PARAMETER_RDG_BUFFER_UPLOAD(InstanceBuffer)
@@ -1042,35 +1041,21 @@ namespace VirtualHeightfieldMesh
 
 	void InitializeInstanceBuffers(FRHICommandListImmediate& InRHICmdList, FDrawInstanceBuffers& InBuffers)
 	{
-		// We use a fake RDG pass for allocation. Is there a more direct way to do this for PooledRDGBuffer objects?
-		// An alternative is use standard RHI allocation, but then we need to be manage resource transitions.
-		FRDGBuilder GraphBuilder(InRHICmdList);
-
-		int32 InstanceBufferSize = CVarVHMMaxRenderItems.GetValueOnRenderThread();
-		FRDGBufferRef InstanceBuffer = GraphBuilder.CreateBuffer(FRDGBufferDesc::CreateStructuredDesc(sizeof(VirtualHeightfieldMesh::QuadRenderInstance), InstanceBufferSize), TEXT("InstanceBuffer"));
-		FRDGBufferRef IndirectArgsBuffer = GraphBuilder.CreateBuffer(FRDGBufferDesc::CreateIndirectDesc(5), TEXT("IndirectArgsBuffer"));
-
-		FCreateBufferParameters* Parameters = GraphBuilder.AllocParameters<FCreateBufferParameters>();
-		Parameters->InstanceBuffer = InstanceBuffer;
-		Parameters->IndirectArgsBuffer = IndirectArgsBuffer;
-
-		GraphBuilder.AddPass(
-			RDG_EVENT_NAME("CreateBuffers"),
-			Parameters,
-			ERDGPassFlags::Copy,
-			[InstanceBuffer, IndirectArgsBuffer](FRHICommandList& RHICmdList)
 			{
-				//InstanceBuffer->MarkResourceAsUsed();
-				//IndirectArgsBuffer->MarkResourceAsUsed();
-			});
-
-		GraphBuilder.QueueBufferExtraction(InstanceBuffer, &InBuffers.InstanceBuffer, ERHIAccess::UAVCompute);
-		GraphBuilder.QueueBufferExtraction(IndirectArgsBuffer, &InBuffers.IndirectArgsBuffer, ERHIAccess::UAVCompute);
-
-		GraphBuilder.Execute();
-
-		// The SRV objects referenced by final rendering are managed outside of RDG.
-		InBuffers.InstanceBufferSRV = RHICreateShaderResourceView(InBuffers.InstanceBuffer->GetRHI());
+			FRHIResourceCreateInfo CreateInfo;
+			CreateInfo.DebugName = TEXT("InstanceBuffer");
+			const int32 InstanceSize = sizeof(VirtualHeightfieldMesh::QuadRenderInstance);
+			const int32 InstanceBufferSize = CVarVHMMaxRenderItems.GetValueOnRenderThread() * InstanceSize;
+			InBuffers.InstanceBuffer = RHICreateStructuredBuffer(InstanceSize, InstanceBufferSize, BUF_UnorderedAccess|BUF_ShaderResource, ERHIAccess::UAVGraphics, CreateInfo);
+			InBuffers.InstanceBufferUAV = RHICreateUnorderedAccessView(InBuffers.InstanceBuffer, false, false);
+			InBuffers.InstanceBufferSRV = RHICreateShaderResourceView(InBuffers.InstanceBuffer);
+		}
+		{
+			FRHIResourceCreateInfo CreateInfo;
+			CreateInfo.DebugName = TEXT("IndirectArgsBuffer");
+			InBuffers.IndirectArgsBuffer = RHICreateVertexBuffer(5 * sizeof(uint32), BUF_UnorderedAccess | BUF_DrawIndirect, ERHIAccess::UAVGraphics, CreateInfo);
+			InBuffers.IndirectArgsBufferUAV = RHICreateUnorderedAccessView(InBuffers.IndirectArgsBuffer, PF_R32_UINT);
+		}
 	}
 
 	/** Initialize the volatie resources used in the render graph. */
@@ -1106,17 +1091,6 @@ namespace VirtualHeightfieldMesh
 		OutResources.QuadNeighborBufferSRV = GraphBuilder.CreateSRV(FRDGBufferSRVDesc(OutResources.QuadNeighborBuffer, PF_R32_UINT));
 	}
 
-	/** Initialize the output resources used in the render graph. */
-	void InitializeResources(FRDGBuilder& GraphBuilder, FDrawInstanceBuffers const& InBuffers, FOutputResources& OutResources)
-	{
-		OutResources.InstanceBuffer = GraphBuilder.RegisterExternalBuffer(InBuffers.InstanceBuffer);
-		OutResources.InstanceBufferUAV = GraphBuilder.CreateUAV(OutResources.InstanceBuffer);
-		OutResources.InstanceBufferSRV = GraphBuilder.CreateSRV(OutResources.InstanceBuffer);
-
-		OutResources.IndirectArgsBuffer = GraphBuilder.RegisterExternalBuffer(InBuffers.IndirectArgsBuffer);
-		OutResources.IndirectArgsBufferUAV = GraphBuilder.CreateUAV(OutResources.IndirectArgsBuffer);
-	}
-
 	/** */
 	void AddPass_InitBuffers(FRDGBuilder& GraphBuilder, FGlobalShaderMap* InGlobalShaderMap, FProxyDesc const& InDesc, FVolatileResources& InVolatileResources)
 	{
@@ -1138,7 +1112,6 @@ namespace VirtualHeightfieldMesh
 			{
 				//todo: If feedback parsing understands append counter we don't need to fully clear
 				RHICmdList.ClearUAVUint(PassParameters->RWFeedbackBuffer->GetRHI(), FUintVector4(0xffffffff, 0xffffffff, 0xffffffff, 0xffffffff));
-				RHICmdList.Transition(FRHITransitionInfo(PassParameters->RWFeedbackBuffer->GetRHI(), ERHIAccess::UAVCompute, ERHIAccess::UAVCompute));
 
 				FComputeShaderUtils::Dispatch(RHICmdList, ComputeShader, *PassParameters, FIntVector(1, 1, 1));
 			});
@@ -1245,7 +1218,7 @@ namespace VirtualHeightfieldMesh
 	}
 
 	/** */
-	void AddPass_InitInstanceBuffer(FRDGBuilder& GraphBuilder, FGlobalShaderMap* InGlobalShaderMap, FProxyDesc const& InDesc, FVolatileResources& InVolatileResources, FOutputResources& InOutputResources)
+	void AddPass_InitInstanceBuffer(FRDGBuilder& GraphBuilder, FGlobalShaderMap* InGlobalShaderMap, FProxyDesc const& InDesc, FVolatileResources& InVolatileResources, FDrawInstanceBuffers& InOutputResources)
 	{
 		TShaderMapRef<FInitInstanceBufferCS> ComputeShader(InGlobalShaderMap);
 
@@ -1260,7 +1233,7 @@ namespace VirtualHeightfieldMesh
 	}
 
 	/** */
-	void AddPass_CullInstances(FRDGBuilder& GraphBuilder, FGlobalShaderMap* InGlobalShaderMap, FProxyDesc const& InDesc, FVolatileResources& InVolatileResources, FOutputResources& InOutputResources, FChildViewDesc const& InViewDesc)
+	void AddPass_CullInstances(FRDGBuilder& GraphBuilder, FGlobalShaderMap* InGlobalShaderMap, FProxyDesc const& InDesc, FVolatileResources& InVolatileResources, FDrawInstanceBuffers& InOutputResources, FChildViewDesc const& InViewDesc)
 	{
 		FCullInstances::FParameters* PassParameters = GraphBuilder.AllocParameters<FCullInstances::FParameters>();
 		PassParameters->MinMaxTexture = InDesc.MinMaxTexture;
@@ -1299,6 +1272,9 @@ namespace VirtualHeightfieldMesh
 			IndirectArgsBuffer->MarkResourceAsUsed();
 			int32 IndirectArgOffset = VirtualHeightfieldMesh::IndirectArgsByteOffset_FinalCull;
 			FComputeShaderUtils::DispatchIndirect(RHICmdList, ComputeShader, *PassParameters, IndirectArgsBuffer->GetIndirectRHICallBuffer(), IndirectArgOffset);
+
+			//RHICmdList.Transition(FRHITransitionInfo(PassParameters->RWIndirectArgsBuffer, ERHIAccess::IndirectArgs));
+			//RHICmdList.Transition(FRHITransitionInfo(PassParameters->RWInstanceBuffer, ERHIAccess::SRVGraphics));
 		});
 	}
 
@@ -1314,7 +1290,7 @@ namespace VirtualHeightfieldMesh
 	}
 
 	/** */
-	void GPUCullChildView(FRDGBuilder& GraphBuilder, FGlobalShaderMap* InGlobalShaderMap, FProxyDesc const& InDesc, FVolatileResources& InVolatileResources, FOutputResources& InOutputResources, FChildViewDesc const& InViewDesc)
+	void GPUCullChildView(FRDGBuilder& GraphBuilder, FGlobalShaderMap* InGlobalShaderMap, FProxyDesc const& InDesc, FVolatileResources& InVolatileResources, FDrawInstanceBuffers& InOutputResources, FChildViewDesc const& InViewDesc)
 	{
 		RDG_EVENT_SCOPE(GraphBuilder, "CullView");
 
@@ -1323,155 +1299,131 @@ namespace VirtualHeightfieldMesh
 	}
 }
 
-void FVirtualHeightfieldMeshRendererExtension::SubmitWork(FRHICommandListImmediate& InRHICmdList)
+void FVirtualHeightfieldMeshRendererExtension::SubmitWork(FRDGBuilder& GraphBuilder)
 {
-	SCOPED_DRAW_EVENT(InRHICmdList, VirtualHeightfieldMesh);
+	DECLARE_GPU_STAT(VirtualHeightfieldMesh)
+	RDG_EVENT_SCOPE(GraphBuilder, "VirtualHeightfieldMesh");
+	RDG_GPU_STAT_SCOPE(GraphBuilder, VirtualHeightfieldMesh);
 
-	// Collect feedback buffers from each pass to submit together after RenderGraph execution.
-	// todo: Convert feedback submission to RDG so that it can be included in the render graph. 
-	//       Then the render graph builder can be passed in and executed externally.
-	TArray< TRefCountPtr< FRDGPooledBuffer > > FeedbackBuffers;
+	// Sort work so that we can batch by proxy/view
+	WorkDescs.Sort(FWorkDescSort());
 
-	FRDGBuilder GraphBuilder(InRHICmdList);
+	const int32 NumWorkItems = WorkDescs.Num();
+	int32 WorkIndex = 0;
+	while (WorkIndex < NumWorkItems)
 	{
-		// Sort work so that we can batch by proxy/view
-		WorkDescs.Sort(FWorkDescSort());
-
-		const int32 NumWorkItems = WorkDescs.Num();
-		int32 WorkIndex = 0;
-		while (WorkIndex < NumWorkItems)
-		{
-			// Gather data per proxy
-			FVirtualHeightfieldMeshSceneProxy const* Proxy = SceneProxies[WorkDescs[WorkIndex].ProxyIndex];
-			IAllocatedVirtualTexture* AllocatedVirtualTexture = Proxy->AllocatedVirtualTexture;
+		// Gather data per proxy
+		FVirtualHeightfieldMeshSceneProxy const* Proxy = SceneProxies[WorkDescs[WorkIndex].ProxyIndex];
+		IAllocatedVirtualTexture* AllocatedVirtualTexture = Proxy->AllocatedVirtualTexture;
 			
-			const float PageSize = AllocatedVirtualTexture->GetVirtualTileSize();
-			const float PageBorderSize = AllocatedVirtualTexture->GetTileBorderSize();
-			const float PageAndBorderSize = PageSize + PageBorderSize * 2.f;
-			const float HalfTexelSize = 0.5f;
-			const float PhysicalTextureSize = AllocatedVirtualTexture->GetPhysicalTextureSize(0);
-			const FVector4 PhysicalPageTransform = FVector4(PageAndBorderSize, PageSize, PageBorderSize, HalfTexelSize) * (1.f / PhysicalTextureSize);
-
-			const float PageTableSizeX = AllocatedVirtualTexture->GetWidthInTiles();
-			const float PageTableSizeY = AllocatedVirtualTexture->GetHeightInTiles();
-			const FVector4 PageTableSize = FVector4(PageTableSizeX, PageTableSizeY, 1.f / PageTableSizeX, 1.f / PageTableSizeY);
-
-			VirtualHeightfieldMesh::FProxyDesc ProxyDesc;
-			ProxyDesc.PageTableTexture = AllocatedVirtualTexture->GetPageTableTexture(0);
-			ProxyDesc.MinMaxTexture = Proxy->MinMaxTexture ? Proxy->MinMaxTexture->Resource->TextureRHI : VirtualHeightfieldMesh::GMinMaxDefaultTexture->TextureRHI;
-			ProxyDesc.MinMaxLevelOffset = ProxyDesc.MinMaxTexture->GetNumMips() - 1 - AllocatedVirtualTexture->GetMaxLevel();
-			ProxyDesc.MaxLevel = AllocatedVirtualTexture->GetMaxLevel();
-			ProxyDesc.PageTableSize = PageTableSize;
-			ProxyDesc.PhysicalPageTransform = PhysicalPageTransform;
-			ProxyDesc.NumPhysicalAddressBits = AllocatedVirtualTexture->GetPageTableFormat() == EVTPageTableFormat::UInt16 ? 6 : 8; // See packing in PageTableUpdate.usf
-			ProxyDesc.PageTableFeedbackId = AllocatedVirtualTexture->GetSpaceID() << 28;
-			ProxyDesc.UVToWorld = Proxy->UVToWorld;
-			ProxyDesc.UVToWorldScale = Proxy->UVToWorldScale;
-			ProxyDesc.NumQuadsPerTileSide = Proxy->NumQuadsPerTileSide;
-			ProxyDesc.MaxPersistentQueueItems = 1 << FMath::CeilLogTwo(CVarVHMMaxPersistentQueueItems.GetValueOnRenderThread());
-			ProxyDesc.MaxRenderItems = CVarVHMMaxRenderItems.GetValueOnRenderThread();
-			ProxyDesc.MaxFeedbackItems = CVarVHMMaxFeedbackItems.GetValueOnRenderThread();
-			ProxyDesc.NumCollectPassWavefronts = CVarVHMCollectPassWavefronts.GetValueOnRenderThread();
-
-			while (WorkIndex < NumWorkItems && SceneProxies[WorkDescs[WorkIndex].ProxyIndex] == Proxy)
-			{
-				// Gather data per main view
-				FSceneView const* MainView = MainViews[WorkDescs[WorkIndex].MainViewIndex];
+		const float PageSize = AllocatedVirtualTexture->GetVirtualTileSize();
+		const float PageBorderSize = AllocatedVirtualTexture->GetTileBorderSize();
+		const float PageAndBorderSize = PageSize + PageBorderSize * 2.f;
+		const float HalfTexelSize = 0.5f;
+		const float PhysicalTextureSize = AllocatedVirtualTexture->GetPhysicalTextureSize(0);
+		const FVector4 PhysicalPageTransform = FVector4(PageAndBorderSize, PageSize, PageBorderSize, HalfTexelSize) * (1.f / PhysicalTextureSize);
+
+		const float PageTableSizeX = AllocatedVirtualTexture->GetWidthInTiles();
+		const float PageTableSizeY = AllocatedVirtualTexture->GetHeightInTiles();
+		const FVector4 PageTableSize = FVector4(PageTableSizeX, PageTableSizeY, 1.f / PageTableSizeX, 1.f / PageTableSizeY);
+
+		VirtualHeightfieldMesh::FProxyDesc ProxyDesc;
+		ProxyDesc.PageTableTexture = AllocatedVirtualTexture->GetPageTableTexture(0);
+		ProxyDesc.MinMaxTexture = Proxy->MinMaxTexture ? Proxy->MinMaxTexture->Resource->TextureRHI : VirtualHeightfieldMesh::GMinMaxDefaultTexture->TextureRHI;
+		ProxyDesc.MinMaxLevelOffset = ProxyDesc.MinMaxTexture->GetNumMips() - 1 - AllocatedVirtualTexture->GetMaxLevel();
+		ProxyDesc.MaxLevel = AllocatedVirtualTexture->GetMaxLevel();
+		ProxyDesc.PageTableSize = PageTableSize;
+		ProxyDesc.PhysicalPageTransform = PhysicalPageTransform;
+		ProxyDesc.NumPhysicalAddressBits = AllocatedVirtualTexture->GetPageTableFormat() == EVTPageTableFormat::UInt16 ? 6 : 8; // See packing in PageTableUpdate.usf
+		ProxyDesc.PageTableFeedbackId = AllocatedVirtualTexture->GetSpaceID() << 28;
+		ProxyDesc.UVToWorld = Proxy->UVToWorld;
+		ProxyDesc.UVToWorldScale = Proxy->UVToWorldScale;
+		ProxyDesc.NumQuadsPerTileSide = Proxy->NumQuadsPerTileSide;
+		ProxyDesc.MaxPersistentQueueItems = 1 << FMath::CeilLogTwo(CVarVHMMaxPersistentQueueItems.GetValueOnRenderThread());
+		ProxyDesc.MaxRenderItems = CVarVHMMaxRenderItems.GetValueOnRenderThread();
+		ProxyDesc.MaxFeedbackItems = CVarVHMMaxFeedbackItems.GetValueOnRenderThread();
+		ProxyDesc.NumCollectPassWavefronts = CVarVHMCollectPassWavefronts.GetValueOnRenderThread();
+
+		while (WorkIndex < NumWorkItems && SceneProxies[WorkDescs[WorkIndex].ProxyIndex] == Proxy)
+		{
+			// Gather data per main view
+			FSceneView const* MainView = MainViews[WorkDescs[WorkIndex].MainViewIndex];
 				
-				VirtualHeightfieldMesh::FViewData MainViewData;
-				VirtualHeightfieldMesh::GetViewData(MainView, MainViewData);
+			VirtualHeightfieldMesh::FViewData MainViewData;
+			VirtualHeightfieldMesh::GetViewData(MainView, MainViewData);
 
-				VirtualHeightfieldMesh::FMainViewDesc MainViewDesc;
-				MainViewDesc.ViewDebug = MainView;
+			VirtualHeightfieldMesh::FMainViewDesc MainViewDesc;
+			MainViewDesc.ViewDebug = MainView;
 
-				// ViewOrigin and Frustum Planes are all converted to UV space for the shader.
-				MainViewDesc.ViewOrigin = Proxy->WorldToUV.TransformPosition(MainViewData.ViewOrigin);
-				MainViewDesc.LodDistances = VirtualHeightfieldMesh::CalculateLodRanges(MainView, Proxy);
+			// ViewOrigin and Frustum Planes are all converted to UV space for the shader.
+			MainViewDesc.ViewOrigin = Proxy->WorldToUV.TransformPosition(MainViewData.ViewOrigin);
+			MainViewDesc.LodDistances = VirtualHeightfieldMesh::CalculateLodRanges(MainView, Proxy);
 
-				const int32 MainViewNumPlanes = FMath::Min(MainViewData.ViewFrustum.Planes.Num(), 5);
-				for (int32 PlaneIndex = 0; PlaneIndex < MainViewNumPlanes; ++PlaneIndex)
+			const int32 MainViewNumPlanes = FMath::Min(MainViewData.ViewFrustum.Planes.Num(), 5);
+			for (int32 PlaneIndex = 0; PlaneIndex < MainViewNumPlanes; ++PlaneIndex)
+			{
+				FPlane Plane = MainViewData.ViewFrustum.Planes[PlaneIndex];
+				Plane = VirtualHeightfieldMesh::TransformPlane(Plane, Proxy->WorldToUV, Proxy->WorldToUVTransposeAdjoint);
+				MainViewDesc.Planes[PlaneIndex] = VirtualHeightfieldMesh::ConvertPlane(Plane);
+			}
+			for (int32 PlaneIndex = MainViewNumPlanes; PlaneIndex < 5; ++PlaneIndex)
+			{
+				MainViewDesc.Planes[PlaneIndex] = FPlane(0, 0, 0, 1); // Null plane won't cull anything
+			}
+
+			FOcclusionResults* OcclusionResults = GOcclusionResults.Find(FOcclusionResultsKey(Proxy, MainView));
+			if (OcclusionResults == nullptr)
+			{
+				MainViewDesc.OcclusionTexture = GBlackTexture->TextureRHI;
+				MainViewDesc.OcclusionLevelOffset = (int32)ProxyDesc.MaxLevel;
+			}
+			else
+			{
+				MainViewDesc.OcclusionTexture = OcclusionResults->OcclusionTexture;
+				MainViewDesc.OcclusionLevelOffset = (int32)ProxyDesc.MaxLevel - OcclusionResults->NumTextureMips + 1;
+			}
+
+			// Build volatile graph resources
+			VirtualHeightfieldMesh::FVolatileResources VolatileResources;
+			VirtualHeightfieldMesh::InitializeResources(GraphBuilder, ProxyDesc, MainViewDesc, VolatileResources);
+
+			// Build graph
+			VirtualHeightfieldMesh::GPUCullMainView(GraphBuilder, GetGlobalShaderMap(GMaxRHIFeatureLevel), ProxyDesc, VolatileResources, MainViewDesc);
+
+			FVirtualTextureFeedbackBufferDesc Desc;
+ 			Desc.Init(CVarVHMMaxFeedbackItems.GetValueOnRenderThread() + 1);
+ 			SubmitVirtualTextureFeedbackBuffer(GraphBuilder, VolatileResources.FeedbackBuffer, Desc);
+
+			while (WorkIndex < NumWorkItems && MainViews[WorkDescs[WorkIndex].MainViewIndex] == MainView)
+			{
+				// Gather data per child view
+				FSceneView const* CullView = CullViews[WorkDescs[WorkIndex].CullViewIndex];
+				FConvexVolume const* ShadowFrustum = CullView->GetDynamicMeshElementsShadowCullFrustum();
+				FConvexVolume const& Frustum = ShadowFrustum != nullptr && ShadowFrustum->Planes.Num() > 0 ? *ShadowFrustum : CullView->ViewFrustum;
+				const FVector PreShadowTranslation = ShadowFrustum != nullptr ? CullView->GetPreShadowTranslation() : FVector::ZeroVector;
+
+				VirtualHeightfieldMesh::FChildViewDesc ChildViewDesc;
+				ChildViewDesc.ViewDebug = MainView;
+				ChildViewDesc.bIsMainView = CullView == MainView;
+					
+				const int32 ChildViewNumPlanes = FMath::Min(Frustum.Planes.Num(), 5);
+				for (int32 PlaneIndex = 0; PlaneIndex < ChildViewNumPlanes; ++PlaneIndex)
 				{
-					FPlane Plane = MainViewData.ViewFrustum.Planes[PlaneIndex];
+					FPlane Plane = Frustum.Planes[PlaneIndex];
+					Plane = VirtualHeightfieldMesh::TranslatePlane(Plane, PreShadowTranslation);
 					Plane = VirtualHeightfieldMesh::TransformPlane(Plane, Proxy->WorldToUV, Proxy->WorldToUVTransposeAdjoint);
-					MainViewDesc.Planes[PlaneIndex] = VirtualHeightfieldMesh::ConvertPlane(Plane);
+					ChildViewDesc.Planes[PlaneIndex] = VirtualHeightfieldMesh::ConvertPlane(Plane);
 				}
-				for (int32 PlaneIndex = MainViewNumPlanes; PlaneIndex < 5; ++PlaneIndex)
+				for (int32 PlaneIndex = ChildViewNumPlanes; PlaneIndex < 5; ++PlaneIndex)
 				{
 					MainViewDesc.Planes[PlaneIndex] = FPlane(0, 0, 0, 1); // Null plane won't cull anything
 				}
 
-				FOcclusionResults* OcclusionResults = GOcclusionResults.Find(FOcclusionResultsKey(Proxy, MainView));
-				if (OcclusionResults == nullptr)
-				{
-					MainViewDesc.OcclusionTexture = GBlackTexture->TextureRHI;
-					MainViewDesc.OcclusionLevelOffset = (int32)ProxyDesc.MaxLevel;
-				}
-				else
-				{
-					MainViewDesc.OcclusionTexture = OcclusionResults->OcclusionTexture;
-					MainViewDesc.OcclusionLevelOffset = (int32)ProxyDesc.MaxLevel - OcclusionResults->NumTextureMips + 1;
-				}
-
-				// Build volatile graph resources
-				VirtualHeightfieldMesh::FVolatileResources VolatileResources;
-				VirtualHeightfieldMesh::InitializeResources(GraphBuilder, ProxyDesc, MainViewDesc, VolatileResources);
-
 				// Build graph
-				VirtualHeightfieldMesh::GPUCullMainView(GraphBuilder, GetGlobalShaderMap(GMaxRHIFeatureLevel), ProxyDesc, VolatileResources, MainViewDesc);
-
-				// Tag feedback buffer for extraction
-				FeedbackBuffers.AddDefaulted();
-				GraphBuilder.QueueBufferExtraction(VolatileResources.FeedbackBuffer, &FeedbackBuffers.Last(), ERHIAccess::SRVMask);
+				VirtualHeightfieldMesh::GPUCullChildView(GraphBuilder, GetGlobalShaderMap(GMaxRHIFeatureLevel), ProxyDesc, VolatileResources, Buffers[WorkDescs[WorkIndex].BufferIndex], ChildViewDesc);
 
-				while (WorkIndex < NumWorkItems && MainViews[WorkDescs[WorkIndex].MainViewIndex] == MainView)
-				{
-					// Gather data per child view
-					FSceneView const* CullView = CullViews[WorkDescs[WorkIndex].CullViewIndex];
-					FConvexVolume const* ShadowFrustum = CullView->GetDynamicMeshElementsShadowCullFrustum();
-					FConvexVolume const& Frustum = ShadowFrustum != nullptr && ShadowFrustum->Planes.Num() > 0 ? *ShadowFrustum : CullView->ViewFrustum;
-					const FVector PreShadowTranslation = ShadowFrustum != nullptr ? CullView->GetPreShadowTranslation() : FVector::ZeroVector;
-
-					VirtualHeightfieldMesh::FChildViewDesc ChildViewDesc;
-					ChildViewDesc.ViewDebug = MainView;
-					ChildViewDesc.bIsMainView = CullView == MainView;
-					
-					const int32 ChildViewNumPlanes = FMath::Min(Frustum.Planes.Num(), 5);
-					for (int32 PlaneIndex = 0; PlaneIndex < ChildViewNumPlanes; ++PlaneIndex)
-					{
-						FPlane Plane = Frustum.Planes[PlaneIndex];
-						Plane = VirtualHeightfieldMesh::TranslatePlane(Plane, PreShadowTranslation);
-						Plane = VirtualHeightfieldMesh::TransformPlane(Plane, Proxy->WorldToUV, Proxy->WorldToUVTransposeAdjoint);
-						ChildViewDesc.Planes[PlaneIndex] = VirtualHeightfieldMesh::ConvertPlane(Plane);
-					}
-					for (int32 PlaneIndex = ChildViewNumPlanes; PlaneIndex < 5; ++PlaneIndex)
-					{
-						MainViewDesc.Planes[PlaneIndex] = FPlane(0, 0, 0, 1); // Null plane won't cull anything
-					}
-
-					// Build output graph resources
-					VirtualHeightfieldMesh::FDrawInstanceBuffers& InstanceBuffers = Buffers[WorkDescs[WorkIndex].BufferIndex];
-					VirtualHeightfieldMesh::FOutputResources OutputResources;
-					VirtualHeightfieldMesh::InitializeResources(GraphBuilder, InstanceBuffers, OutputResources);
-
-					// Build graph
-					VirtualHeightfieldMesh::GPUCullChildView(GraphBuilder, GetGlobalShaderMap(GMaxRHIFeatureLevel), ProxyDesc, VolatileResources, OutputResources, ChildViewDesc);
-
-					WorkIndex++;
-				}
+				WorkIndex++;
 			}
 		}
 	}
-
-	GraphBuilder.Execute();
-
-	// Subit feedback buffers
-	{
-		SCOPED_DRAW_EVENT(InRHICmdList, CopyVirtualTextureFeedback);
-		for (int32 FeedbackIndex = 0; FeedbackIndex < FeedbackBuffers.Num(); ++FeedbackIndex)
-		{
-			FVirtualTextureFeedbackBufferDesc Desc;
-			Desc.Init(CVarVHMMaxFeedbackItems.GetValueOnRenderThread() + 1);
-			SubmitVirtualTextureFeedbackBuffer(GraphBuilder.RHICmdList, FeedbackBuffers[0].GetReference()->GetRHI(), Desc);
-		}
-	}
 }
