Fix VirtualHeightfieldMesh shader for Vulkan and Metal RHI. (RW-)StructuredBuffer must be 16-byte aligned in DXC's SPIR-V backend.

#rnx

diff --git a/Engine/Plugins/Experimental/VirtualHeightfieldMesh/Shaders/Private/VirtualHeightfieldMesh.usf b/Engine/Plugins/Experimental/VirtualHeightfieldMesh/Shaders/Private/VirtualHeightfieldMesh.usf
--- a/Engine/Plugins/Experimental/VirtualHeightfieldMesh/Shaders/Private/VirtualHeightfieldMesh.usf
+++ b/Engine/Plugins/Experimental/VirtualHeightfieldMesh/Shaders/Private/VirtualHeightfieldMesh.usf
@@ -492,7 +492,7 @@ void CullInstancesCS( uint3 DispatchThreadId : SV_DispatchThreadID )
 			const uint PackedNeighborItem = QuadNeighborBuffer[QuadIndex * 4 + NeighborIndex];
 			const QuadNeighborItem NeighborItem = UnpackQuadNeighborItem(PackedNeighborItem);
 
-			OutInstance.NeigborUVTransform[NeighborIndex] = GetVirtualToPhysicalUVTransform(Pos, Level, NeighborItem.PhysicalAddress, PhysicalPageTransform, NumPhysicalAddressBits);
+			OutInstance.NeigborUVTransform[NeighborIndex] = float4(GetVirtualToPhysicalUVTransform(Pos, Level, NeighborItem.PhysicalAddress, PhysicalPageTransform, NumPhysicalAddressBits), 0);
 		}
 
 		uint Write;
diff --git a/Engine/Plugins/Experimental/VirtualHeightfieldMesh/Shaders/Private/VirtualHeightfieldMesh.ush b/Engine/Plugins/Experimental/VirtualHeightfieldMesh/Shaders/Private/VirtualHeightfieldMesh.ush
--- a/Engine/Plugins/Experimental/VirtualHeightfieldMesh/Shaders/Private/VirtualHeightfieldMesh.ush
+++ b/Engine/Plugins/Experimental/VirtualHeightfieldMesh/Shaders/Private/VirtualHeightfieldMesh.ush
@@ -124,7 +124,7 @@ struct QuadRenderInstance
 	// Total of 64 bytes is nice for cache alignment. Take care when modifying.
 	uint PosLevelPacked;
 	float3 UVTransform;
-	float3 NeigborUVTransform[4];
+	float4 NeigborUVTransform[4];
 };
 
 uint2 UnpackPos(QuadRenderInstance Item)
diff --git a/Engine/Plugins/Experimental/VirtualHeightfieldMesh/Source/VirtualHeightfieldMesh/Private/VirtualHeightfieldMeshSceneProxy.cpp b/Engine/Plugins/Experimental/VirtualHeightfieldMesh/Source/VirtualHeightfieldMesh/Private/VirtualHeightfieldMeshSceneProxy.cpp
--- a/Engine/Plugins/Experimental/VirtualHeightfieldMesh/Source/VirtualHeightfieldMesh/Private/VirtualHeightfieldMeshSceneProxy.cpp
+++ b/Engine/Plugins/Experimental/VirtualHeightfieldMesh/Source/VirtualHeightfieldMesh/Private/VirtualHeightfieldMeshSceneProxy.cpp
@@ -668,7 +668,7 @@ namespace VirtualHeightfieldMesh
 	{
 		uint32 AddressLevelPacked;
 		float UVTransform[3];
-		float NeigborUVTransform[4][3];
+		float NeigborUVTransform[4][4];
 	};
 
 	/** Compute shader to initialize all buffers, including adding the lowest mip page(s) to the QuadBuffer. */
