Fix transition validation error for VHM

diff --git a/Engine/Plugins/Experimental/VirtualHeightfieldMesh/Source/VirtualHeightfieldMesh/Private/VirtualHeightfieldMeshSceneProxy.cpp b/Engine/Plugins/Experimental/VirtualHeightfieldMesh/Source/VirtualHeightfieldMesh/Private/VirtualHeightfieldMeshSceneProxy.cpp
--- a/Engine/Plugins/Experimental/VirtualHeightfieldMesh/Source/VirtualHeightfieldMesh/Private/VirtualHeightfieldMeshSceneProxy.cpp
+++ b/Engine/Plugins/Experimental/VirtualHeightfieldMesh/Source/VirtualHeightfieldMesh/Private/VirtualHeightfieldMeshSceneProxy.cpp
@@ -135,7 +135,7 @@ struct FOcclusionResultsKey
 	}
 };
 
-/** */
+/** Global map for occlusion result. */
 TMap< FOcclusionResultsKey, FOcclusionResults > GOcclusionResults;
 
 namespace VirtualHeightfieldMesh
@@ -740,7 +740,7 @@ namespace VirtualHeightfieldMesh
 			RENDER_TARGET_BINDING_SLOTS()
 			SHADER_PARAMETER(FVector4, PageTableSize)
 			SHADER_PARAMETER_RDG_BUFFER_SRV(Buffer<uint2>, QuadBuffer)
-			SHADER_PARAMETER_RDG_BUFFER(Buffer<uint>, IndirectArgsBuffer)
+			RDG_BUFFER_ACCESS(IndirectArgsBuffer, ERHIAccess::IndirectArgs)
 		END_SHADER_PARAMETER_STRUCT()
 
 		static bool ShouldCompilePermutation(FGlobalShaderPermutationParameters const& Parameters)
@@ -781,7 +781,7 @@ namespace VirtualHeightfieldMesh
 
 	IMPLEMENT_GLOBAL_SHADER(FRenderLodMapPS, "/Plugin/VirtualHeightfieldMesh/Private/VirtualHeightfieldMesh.usf", "RenderLodMapPS", SF_Pixel);
 
-	/** */
+	/** ResolveNeighborLods compute shader. */
 	class FResolveNeighborLodsCS : public FGlobalShader
 	{
 	public:
@@ -794,10 +794,10 @@ namespace VirtualHeightfieldMesh
 			SHADER_PARAMETER(uint32, PageTableFeedbackId)
 			SHADER_PARAMETER_RDG_BUFFER_SRV(Buffer<uint2>, QuadBuffer)
 			SHADER_PARAMETER_RDG_TEXTURE(Texture2D<float2>, LodTexture)
-			SHADER_PARAMETER_RDG_BUFFER(Buffer<uint>, IndirectArgsBuffer)
 			SHADER_PARAMETER_RDG_BUFFER_SRV(Buffer<uint>, IndirectArgsBufferSRV)
 			SHADER_PARAMETER_RDG_BUFFER_UAV(RWBuffer<uint>, RWQuadNeighborBuffer)
 			SHADER_PARAMETER_RDG_BUFFER_UAV(RWBuffer<uint>, RWFeedbackBuffer)
+			RDG_BUFFER_ACCESS(IndirectArgsBuffer, ERHIAccess::IndirectArgs)
 		END_SHADER_PARAMETER_STRUCT()
 
 		static bool ShouldCompilePermutation(FGlobalShaderPermutationParameters const& Parameters)
@@ -808,7 +808,7 @@ namespace VirtualHeightfieldMesh
 
 	IMPLEMENT_GLOBAL_SHADER(FResolveNeighborLodsCS, "/Plugin/VirtualHeightfieldMesh/Private/VirtualHeightfieldMesh.usf", "ResolveNeighborLodsCS", SF_Compute);
 
-	/** */
+	/** InitInstanceBuffer compute shader. */
 	class FInitInstanceBufferCS : public FGlobalShader
 	{
 	public:
@@ -828,7 +828,7 @@ namespace VirtualHeightfieldMesh
 
 	IMPLEMENT_GLOBAL_SHADER(FInitInstanceBufferCS, "/Plugin/VirtualHeightfieldMesh/Private/VirtualHeightfieldMesh.usf", "InitInstanceBufferCS", SF_Compute);
 
-	/** */
+	/** CullInstances compute shader. */
 	class FCullInstances : public FGlobalShader
 	{
 	public:
@@ -844,10 +844,10 @@ namespace VirtualHeightfieldMesh
 			SHADER_PARAMETER(uint32, NumPhysicalAddressBits)
 			SHADER_PARAMETER_RDG_BUFFER_SRV(Buffer<uint2>, QuadBuffer)
 			SHADER_PARAMETER_RDG_BUFFER_SRV(Buffer<uint>, QuadNeighborBuffer)
-			SHADER_PARAMETER_RDG_BUFFER(Buffer<uint>, IndirectArgsBuffer)
 			SHADER_PARAMETER_RDG_BUFFER_SRV(Buffer<uint>, IndirectArgsBufferSRV)
 			SHADER_PARAMETER_UAV(RWStructuredBuffer<QuadRenderInstance>, RWInstanceBuffer)
 			SHADER_PARAMETER_UAV(RWBuffer<uint>, RWIndirectArgsBuffer)
+			RDG_BUFFER_ACCESS(IndirectArgsBuffer, ERHIAccess::IndirectArgs)
 		END_SHADER_PARAMETER_STRUCT()
 	};
 
@@ -909,7 +909,7 @@ namespace VirtualHeightfieldMesh
 	/** Single global instance of default Min/Max texture. */
 	FTexture* GMinMaxDefaultTexture = new TGlobalResource<FMinMaxDefaultTexture>;
 
-	/** */
+	/** View matrices that can be frozen in freezerendering mode. */
 	struct FViewData
 	{
 		FVector ViewOrigin;
@@ -1041,19 +1041,19 @@ namespace VirtualHeightfieldMesh
 
 	void InitializeInstanceBuffers(FRHICommandListImmediate& InRHICmdList, FDrawInstanceBuffers& InBuffers)
 	{
-			{
+		{
 			FRHIResourceCreateInfo CreateInfo;
 			CreateInfo.DebugName = TEXT("InstanceBuffer");
 			const int32 InstanceSize = sizeof(VirtualHeightfieldMesh::QuadRenderInstance);
 			const int32 InstanceBufferSize = CVarVHMMaxRenderItems.GetValueOnRenderThread() * InstanceSize;
-			InBuffers.InstanceBuffer = RHICreateStructuredBuffer(InstanceSize, InstanceBufferSize, BUF_UnorderedAccess|BUF_ShaderResource, ERHIAccess::UAVGraphics, CreateInfo);
+			InBuffers.InstanceBuffer = RHICreateStructuredBuffer(InstanceSize, InstanceBufferSize, BUF_UnorderedAccess|BUF_ShaderResource, ERHIAccess::SRVMask, CreateInfo);
 			InBuffers.InstanceBufferUAV = RHICreateUnorderedAccessView(InBuffers.InstanceBuffer, false, false);
 			InBuffers.InstanceBufferSRV = RHICreateShaderResourceView(InBuffers.InstanceBuffer);
 		}
 		{
 			FRHIResourceCreateInfo CreateInfo;
 			CreateInfo.DebugName = TEXT("IndirectArgsBuffer");
-			InBuffers.IndirectArgsBuffer = RHICreateVertexBuffer(5 * sizeof(uint32), BUF_UnorderedAccess | BUF_DrawIndirect, ERHIAccess::UAVGraphics, CreateInfo);
+			InBuffers.IndirectArgsBuffer = RHICreateVertexBuffer(5 * sizeof(uint32), BUF_UnorderedAccess|BUF_DrawIndirect, ERHIAccess::IndirectArgs, CreateInfo);
 			InBuffers.IndirectArgsBufferUAV = RHICreateUnorderedAccessView(InBuffers.IndirectArgsBuffer, PF_R32_UINT);
 		}
 	}
@@ -1091,7 +1091,43 @@ namespace VirtualHeightfieldMesh
 		OutResources.QuadNeighborBufferSRV = GraphBuilder.CreateSRV(FRDGBufferSRVDesc(OutResources.QuadNeighborBuffer, PF_R32_UINT));
 	}
 
-	/** */
+	/** Transition our output draw buffers for use. Read or write access is set according to the bToWrite parameter. */
+	void AddPass_TransitionAllDrawBuffers(FRDGBuilder& GraphBuilder, TArray<VirtualHeightfieldMesh::FDrawInstanceBuffers> const& Buffers, TArrayView<int32> const& BufferIndices, bool bToWrite)
+	{
+		TArray<FRHIUnorderedAccessView*> OverlapUAVs;
+		OverlapUAVs.Reserve(BufferIndices.Num());
+
+		TArray<FRHITransitionInfo> TransitionInfos;
+		TransitionInfos.Reserve(BufferIndices.Num() * 2);
+		
+		for (int32 BufferIndex : BufferIndices)
+		{
+			FRHIUnorderedAccessView* IndirectArgsBufferUAV = Buffers[BufferIndex].IndirectArgsBufferUAV;
+			FRHIUnorderedAccessView* InstanceBufferUAV = Buffers[BufferIndex].InstanceBufferUAV;
+
+			OverlapUAVs.Add(IndirectArgsBufferUAV);
+
+			TransitionInfos.Add(FRHITransitionInfo(IndirectArgsBufferUAV, bToWrite ? ERHIAccess::IndirectArgs : ERHIAccess::UAVMask, bToWrite ? ERHIAccess::UAVMask : ERHIAccess::IndirectArgs));
+			TransitionInfos.Add(FRHITransitionInfo(InstanceBufferUAV, bToWrite ? ERHIAccess::SRVMask : ERHIAccess::UAVMask, bToWrite ? ERHIAccess::UAVMask : ERHIAccess::SRVMask));
+		}
+
+		AddPass(GraphBuilder, [bToWrite, OverlapUAVs, TransitionInfos](FRHICommandListImmediate& InRHICmdList)
+		{
+			if (!bToWrite)
+			{
+				InRHICmdList.EndUAVOverlap(OverlapUAVs);
+			}
+
+			InRHICmdList.Transition(TransitionInfos);
+			
+			if (bToWrite)
+			{
+				InRHICmdList.BeginUAVOverlap(OverlapUAVs);
+			}
+		});
+	}
+
+	/** Initialize the buffers before collecting visible quads. */
 	void AddPass_InitBuffers(FRDGBuilder& GraphBuilder, FGlobalShaderMap* InGlobalShaderMap, FProxyDesc const& InDesc, FVolatileResources& InVolatileResources)
 	{
 		TShaderMapRef<FInitBuffersCS> ComputeShader(InGlobalShaderMap);
@@ -1109,15 +1145,15 @@ namespace VirtualHeightfieldMesh
 			PassParameters,
 			ERDGPassFlags::Compute,
 			[PassParameters, ComputeShader](FRHICommandList& RHICmdList)
-			{
-				//todo: If feedback parsing understands append counter we don't need to fully clear
-				RHICmdList.ClearUAVUint(PassParameters->RWFeedbackBuffer->GetRHI(), FUintVector4(0xffffffff, 0xffffffff, 0xffffffff, 0xffffffff));
+		{
+			//todo: If feedback parsing understands append counter we don't need to fully clear
+			RHICmdList.ClearUAVUint(PassParameters->RWFeedbackBuffer->GetRHI(), FUintVector4(0xffffffff, 0xffffffff, 0xffffffff, 0xffffffff));
 
-				FComputeShaderUtils::Dispatch(RHICmdList, ComputeShader, *PassParameters, FIntVector(1, 1, 1));
-			});
+			FComputeShaderUtils::Dispatch(RHICmdList, ComputeShader, *PassParameters, FIntVector(1, 1, 1));
+		});
 	}
 
-	/** */
+	/** Collect potentially visible quads and determine their Lods. */
 	void AddPass_CollectQuads(FRDGBuilder& GraphBuilder, FGlobalShaderMap* InGlobalShaderMap, FProxyDesc const& InDesc, FVolatileResources& InVolatileResources, FMainViewDesc const& InViewDesc)
 	{
 		TShaderMapRef<FCollectQuadsCS> ComputeShader(InGlobalShaderMap);
@@ -1150,7 +1186,7 @@ namespace VirtualHeightfieldMesh
 			ComputeShader, PassParameters, FIntVector(InDesc.NumCollectPassWavefronts, 1, 1));
 	}
 
-	/**  */
+	/** Render Lod info map, which can then be sampled to determine neighbor Lod transitions. */
 	void AddPass_RenderLodMap(FRDGBuilder& GraphBuilder, FGlobalShaderMap* InGlobalShaderMap, FProxyDesc const& InDesc, FVolatileResources& InVolatileResources)
 	{
 		TShaderMapRef< FRenderLodMapVS > VertexShader(InGlobalShaderMap);
@@ -1167,7 +1203,7 @@ namespace VirtualHeightfieldMesh
 			PassParameters,
 			ERDGPassFlags::Raster,
 			[VertexShader, PixelShader, PassParameters, IndirectArgsBuffer = InVolatileResources.IndirectArgsBuffer](FRHICommandListImmediate& RHICmdListImmediate)
-			{
+		{
 				FGraphicsPipelineStateInitializer GraphicsPSOInit;
 				RHICmdListImmediate.ApplyCachedRenderTargets(GraphicsPSOInit);
 				GraphicsPSOInit.DepthStencilState = TStaticDepthStencilState<false, CF_Always>::GetRHI();
@@ -1182,14 +1218,13 @@ namespace VirtualHeightfieldMesh
 				SetShaderParameters(RHICmdListImmediate, VertexShader, VertexShader.GetVertexShader(), *PassParameters);
 				SetShaderParameters(RHICmdListImmediate, PixelShader, PixelShader.GetPixelShader(), *PassParameters);
 
-				IndirectArgsBuffer->MarkResourceAsUsed();
-
 				int32 IndirectArgOffset = VirtualHeightfieldMesh::IndirectArgsByteOffset_RenderLodMap;
 				RHICmdListImmediate.DrawIndexedPrimitiveIndirect(GTwoTrianglesIndexBuffer.IndexBufferRHI, IndirectArgsBuffer->GetIndirectRHICallBuffer(), IndirectArgOffset);
-			});
+				IndirectArgsBuffer->MarkResourceAsUsed();
+		});
 	}
 
-	/** */
+	/** Resolve Neighbor Lod transitions from the Lod map. */
 	void AddPass_ResolveNeighborLods(FRDGBuilder& GraphBuilder, FGlobalShaderMap* InGlobalShaderMap, FProxyDesc const& InDesc, FVolatileResources& InVolatileResources, FMainViewDesc const& InViewDesc)
 	{
 		TShaderMapRef<FResolveNeighborLodsCS> ComputeShader(InGlobalShaderMap);
@@ -1210,29 +1245,29 @@ namespace VirtualHeightfieldMesh
 			PassParameters,
 			ERDGPassFlags::Compute,
 			[PassParameters, ComputeShader, IndirectArgsBuffer = InVolatileResources.IndirectArgsBuffer](FRHICommandList& RHICmdList)
-			{
-				IndirectArgsBuffer->MarkResourceAsUsed();
+		{
 				int32 IndirectArgOffset = VirtualHeightfieldMesh::IndirectArgsByteOffset_FetchNeighborLod;
 				FComputeShaderUtils::DispatchIndirect(RHICmdList, ComputeShader, *PassParameters, IndirectArgsBuffer->GetIndirectRHICallBuffer(), IndirectArgOffset);
-			});
+				IndirectArgsBuffer->MarkResourceAsUsed();
+		});
 	}
 
-	/** */
-	void AddPass_InitInstanceBuffer(FRDGBuilder& GraphBuilder, FGlobalShaderMap* InGlobalShaderMap, FProxyDesc const& InDesc, FVolatileResources& InVolatileResources, FDrawInstanceBuffers& InOutputResources)
+	/** Initialise the draw indirect buffer. */
+	void AddPass_InitInstanceBuffer(FRDGBuilder& GraphBuilder, FGlobalShaderMap* InGlobalShaderMap, int32 NumQuadsPerTileSide, FDrawInstanceBuffers& InOutputResources)
 	{
 		TShaderMapRef<FInitInstanceBufferCS> ComputeShader(InGlobalShaderMap);
 
 		FInitInstanceBufferCS::FParameters* PassParameters = GraphBuilder.AllocParameters<FInitInstanceBufferCS::FParameters>();
-		PassParameters->NumIndices = InDesc.NumQuadsPerTileSide * InDesc.NumQuadsPerTileSide * 6;
+		PassParameters->NumIndices = NumQuadsPerTileSide * NumQuadsPerTileSide * 6;
 		PassParameters->RWIndirectArgsBuffer = InOutputResources.IndirectArgsBufferUAV;
 
-		FComputeShaderUtils::AddPass(
-			GraphBuilder,
-			RDG_EVENT_NAME("InitInstanceBuffer"),
-			ComputeShader, PassParameters, FIntVector(1, 1, 1));
+ 		FComputeShaderUtils::AddPass(
+ 			GraphBuilder,
+ 			RDG_EVENT_NAME("InitInstanceBuffer"),
+ 			ComputeShader, PassParameters, FIntVector(1, 1, 1));
 	}
 
-	/** */
+	/** Cull quads and write to the final output buffer. */
 	void AddPass_CullInstances(FRDGBuilder& GraphBuilder, FGlobalShaderMap* InGlobalShaderMap, FProxyDesc const& InDesc, FVolatileResources& InVolatileResources, FDrawInstanceBuffers& InOutputResources, FChildViewDesc const& InViewDesc)
 	{
 		FCullInstances::FParameters* PassParameters = GraphBuilder.AllocParameters<FCullInstances::FParameters>();
@@ -1269,16 +1304,13 @@ namespace VirtualHeightfieldMesh
 			ERDGPassFlags::Compute,
 			[PassParameters, ComputeShader, IndirectArgsBuffer = InVolatileResources.IndirectArgsBuffer](FRHICommandList& RHICmdList)
 		{
-			IndirectArgsBuffer->MarkResourceAsUsed();
 			int32 IndirectArgOffset = VirtualHeightfieldMesh::IndirectArgsByteOffset_FinalCull;
 			FComputeShaderUtils::DispatchIndirect(RHICmdList, ComputeShader, *PassParameters, IndirectArgsBuffer->GetIndirectRHICallBuffer(), IndirectArgOffset);
-
-			//RHICmdList.Transition(FRHITransitionInfo(PassParameters->RWIndirectArgsBuffer, ERHIAccess::IndirectArgs));
-			//RHICmdList.Transition(FRHITransitionInfo(PassParameters->RWInstanceBuffer, ERHIAccess::SRVGraphics));
+			IndirectArgsBuffer->MarkResourceAsUsed();
 		});
 	}
 
-	/** */
+	/** Add passes for setting up culling for a main view. */
 	void GPUCullMainView( FRDGBuilder& GraphBuilder, FGlobalShaderMap* InGlobalShaderMap, FProxyDesc const& InDesc, FVolatileResources& InVolatileResources, FMainViewDesc const& InViewDesc)
 	{
 		RDG_EVENT_SCOPE(GraphBuilder, "MainView");
@@ -1289,12 +1321,11 @@ namespace VirtualHeightfieldMesh
 		AddPass_ResolveNeighborLods(GraphBuilder, InGlobalShaderMap, InDesc, InVolatileResources, InViewDesc);
 	}
 
-	/** */
+	/** Add passes for culling a single child view. */
 	void GPUCullChildView(FRDGBuilder& GraphBuilder, FGlobalShaderMap* InGlobalShaderMap, FProxyDesc const& InDesc, FVolatileResources& InVolatileResources, FDrawInstanceBuffers& InOutputResources, FChildViewDesc const& InViewDesc)
 	{
 		RDG_EVENT_SCOPE(GraphBuilder, "CullView");
 
-		AddPass_InitInstanceBuffer(GraphBuilder, InGlobalShaderMap, InDesc, InVolatileResources, InOutputResources);
 		AddPass_CullInstances(GraphBuilder, InGlobalShaderMap, InDesc, InVolatileResources, InOutputResources, InViewDesc);
 	}
 }
@@ -1308,6 +1339,22 @@ void FVirtualHeightfieldMeshRendererExtension::SubmitWork(FRDGBuilder& GraphBuil
 	// Sort work so that we can batch by proxy/view
 	WorkDescs.Sort(FWorkDescSort());
 
+	// Add pass to transition all output buffers for writing
+	TArray<int32, TInlineAllocator<8>> UsedBufferIndices;
+	for (FWorkDesc WorkdDesc : WorkDescs)
+	{
+		UsedBufferIndices.Add(WorkdDesc.BufferIndex);
+	}
+	AddPass_TransitionAllDrawBuffers(GraphBuilder, Buffers, UsedBufferIndices, true);
+
+	// Add passes to initialize the output buffers
+	for (FWorkDesc WorkDesc : WorkDescs)
+	{
+		const int32 NumQuadsPerTileSide = SceneProxies[WorkDesc.ProxyIndex]->NumQuadsPerTileSide;
+		AddPass_InitInstanceBuffer(GraphBuilder, GetGlobalShaderMap(GMaxRHIFeatureLevel), NumQuadsPerTileSide, Buffers[WorkDesc.BufferIndex]);
+	}
+
+	// Iterate workloads and submit work
 	const int32 NumWorkItems = WorkDescs.Num();
 	int32 WorkIndex = 0;
 	while (WorkIndex < NumWorkItems)
@@ -1426,4 +1473,7 @@ void FVirtualHeightfieldMeshRendererExtension::SubmitWork(FRDGBuilder& GraphBuil
 			}
 		}
 	}
+
+	// Add pass to transition all output buffers for reading
+	AddPass_TransitionAllDrawBuffers(GraphBuilder, Buffers, UsedBufferIndices, false);
 }
