Fix shader compilation warning

diff --git a/Engine/Plugins/Experimental/VirtualHeightfieldMesh/Shaders/Private/VirtualHeightfieldMeshVertexFactory.ush b/Engine/Plugins/Experimental/VirtualHeightfieldMesh/Shaders/Private/VirtualHeightfieldMeshVertexFactory.ush
--- a/Engine/Plugins/Experimental/VirtualHeightfieldMesh/Shaders/Private/VirtualHeightfieldMeshVertexFactory.ush
+++ b/Engine/Plugins/Experimental/VirtualHeightfieldMesh/Shaders/Private/VirtualHeightfieldMeshVertexFactory.ush
@@ -83,7 +83,7 @@ FVertexFactoryIntermediates GetVertexFactoryIntermediates(FVertexFactoryInput In
 		(LocalUV.x < LocalUV.y ? 0.5 - LocalUV.x : 0.5 - LocalUV.y);
 	NeighborWeight *= 2;
 
-	float3 NeigborUVTransform = Item.NeigborUVTransform[NeighborIndex];
+	float3 NeigborUVTransform = Item.NeigborUVTransform[NeighborIndex].xyz;
 	float2 NeighborPhysicalUV = NeigborUVTransform.xy + LocalUV * NeigborUVTransform.z;
 
 	// Calculate vertex UV details before morphing
