Quick fix for occasional crash in VHM from occlusion textures released before RDG executes.

diff --git a/Engine/Plugins/Experimental/VirtualHeightfieldMesh/Source/VirtualHeightfieldMesh/Private/VirtualHeightfieldMeshSceneProxy.cpp b/Engine/Plugins/Experimental/VirtualHeightfieldMesh/Source/VirtualHeightfieldMesh/Private/VirtualHeightfieldMeshSceneProxy.cpp
--- a/Engine/Plugins/Experimental/VirtualHeightfieldMesh/Source/VirtualHeightfieldMesh/Private/VirtualHeightfieldMeshSceneProxy.cpp
+++ b/Engine/Plugins/Experimental/VirtualHeightfieldMesh/Source/VirtualHeightfieldMesh/Private/VirtualHeightfieldMeshSceneProxy.cpp
@@ -17,7 +17,8 @@
 #include "VT/RuntimeVirtualTexture.h"
 #include "VT/VirtualTextureFeedbackBuffer.h"
 
-DECLARE_STATS_GROUP(TEXT("Virtual Heightfield Mesh"), STATGROUP_VirtualHeightfieldMesh, STATCAT_Advanced);
+DECLARE_STATS_GROUP(TEXT("VirtualHeightfieldMesh"), STATGROUP_VirtualHeightfieldMesh, STATCAT_Advanced);
+DECLARE_CYCLE_STAT(TEXT("VirtualHeightfieldMesh SubmitWork"), STAT_VirtualHeightfieldMesh_SubmitWork, STATGROUP_VirtualHeightfieldMesh);
 
 DECLARE_LOG_CATEGORY_EXTERN(LogVirtualHeightfieldMesh, Warning, All);
 DEFINE_LOG_CATEGORY(LogVirtualHeightfieldMesh);
@@ -137,6 +138,7 @@ struct FOcclusionResultsKey
 
 /** Global map for occlusion result. */
 TMap< FOcclusionResultsKey, FOcclusionResults > GOcclusionResults;
+bool GOcclusionResetRequired = false;
 
 namespace VirtualHeightfieldMesh
 {
@@ -345,7 +347,7 @@ void FVirtualHeightfieldMeshRendererExtension::EndFrame()
 		}
 	}
 
-	GOcclusionResults.Reset();
+	GOcclusionResetRequired = true;
 }
 
 void FVirtualHeightfieldMeshRendererExtension::EndFrame(FRDGBuilder& GraphBuilder)
@@ -615,6 +617,12 @@ void FVirtualHeightfieldMeshSceneProxy::AcceptOcclusionResults(FSceneView const*
 {
 	check(IsInRenderingThread());
 
+	if (GOcclusionResetRequired)
+	{
+		GOcclusionResults.Reset();
+		GOcclusionResetRequired = false;
+	}
+
 	if (CVarVHMOcclusion.GetValueOnAnyThread() != 0 && Results != nullptr && NumResults > 1)
 	{
 		FOcclusionResults& OcclusionResults = GOcclusionResults.Emplace(FOcclusionResultsKey(this, View));
@@ -1033,12 +1041,7 @@ namespace VirtualHeightfieldMesh
 		FRDGBufferSRVRef QuadNeighborBufferSRV;
 	};
 
-	/* Dummy parameter struct used to allocate FPooledRDGBuffer objects using a fake RDG pass. */
-	BEGIN_SHADER_PARAMETER_STRUCT(FCreateBufferParameters, )
-		SHADER_PARAMETER_RDG_BUFFER_UPLOAD(InstanceBuffer)
-		SHADER_PARAMETER_RDG_BUFFER_UPLOAD(IndirectArgsBuffer)
-	END_SHADER_PARAMETER_STRUCT()
-
+	/** Initialize the FDrawInstanceBuffers objects. */
 	void InitializeInstanceBuffers(FRHICommandListImmediate& InRHICmdList, FDrawInstanceBuffers& InBuffers)
 	{
 		{
@@ -1058,7 +1061,7 @@ namespace VirtualHeightfieldMesh
 		}
 	}
 
-	/** Initialize the volatie resources used in the render graph. */
+	/** Initialize the volatile resources used in the render graph. */
 	void InitializeResources(FRDGBuilder& GraphBuilder, FProxyDesc const& InDesc, FMainViewDesc const& InMainViewDesc, FVolatileResources& OutResources)
 	{
 		OutResources.QueueInfo = GraphBuilder.CreateBuffer(FRDGBufferDesc::CreateStructuredDesc(sizeof(WorkerQueueInfo), 1), TEXT("QueueInfo"));
@@ -1332,6 +1335,7 @@ namespace VirtualHeightfieldMesh
 
 void FVirtualHeightfieldMeshRendererExtension::SubmitWork(FRDGBuilder& GraphBuilder)
 {
+	SCOPE_CYCLE_COUNTER(STAT_VirtualHeightfieldMesh_SubmitWork);
 	DECLARE_GPU_STAT(VirtualHeightfieldMesh)
 	RDG_EVENT_SCOPE(GraphBuilder, "VirtualHeightfieldMesh");
 	RDG_GPU_STAT_SCOPE(GraphBuilder, VirtualHeightfieldMesh);
