VHM defaults to CastShadow

diff --git a/Engine/Plugins/Experimental/VirtualHeightfieldMesh/Source/VirtualHeightfieldMesh/Private/VirtualHeightfieldMeshComponent.cpp b/Engine/Plugins/Experimental/VirtualHeightfieldMesh/Source/VirtualHeightfieldMesh/Private/VirtualHeightfieldMeshComponent.cpp
--- a/Engine/Plugins/Experimental/VirtualHeightfieldMesh/Source/VirtualHeightfieldMesh/Private/VirtualHeightfieldMeshComponent.cpp
+++ b/Engine/Plugins/Experimental/VirtualHeightfieldMesh/Source/VirtualHeightfieldMesh/Private/VirtualHeightfieldMeshComponent.cpp
@@ -13,6 +13,8 @@
 UVirtualHeightfieldMeshComponent::UVirtualHeightfieldMeshComponent(const FObjectInitializer& ObjectInitializer)
 	: Super(ObjectInitializer)
 {
+	CastShadow = true;
+	bUseAsOccluder = true;
 	bNeverDistanceCull = true;
 #if WITH_EDITORONLY_DATA
 	bEnableAutoLODGeneration = false;
