VHM updates after the associated RVT is moved.
Fixes issue where VHM scene proxy has incorrect transform and can disappear.

diff --git a/Engine/Plugins/Experimental/VirtualHeightfieldMesh/Source/VirtualHeightfieldMesh/Private/VirtualHeightfieldMeshComponent.cpp b/Engine/Plugins/Experimental/VirtualHeightfieldMesh/Source/VirtualHeightfieldMesh/Private/VirtualHeightfieldMeshComponent.cpp
--- a/Engine/Plugins/Experimental/VirtualHeightfieldMesh/Source/VirtualHeightfieldMesh/Private/VirtualHeightfieldMeshComponent.cpp
+++ b/Engine/Plugins/Experimental/VirtualHeightfieldMesh/Source/VirtualHeightfieldMesh/Private/VirtualHeightfieldMeshComponent.cpp
@@ -29,6 +29,8 @@ void UVirtualHeightfieldMeshComponent::OnRegister()
 	URuntimeVirtualTextureComponent* RuntimeVirtualTextureComponent = VirtualTextureRef != nullptr ? VirtualTextureRef->VirtualTextureComponent : nullptr;
 	if (RuntimeVirtualTextureComponent)
 	{
+		// Bind to delegate so that we dirty render state whenever RuntimeVirtualTextureComponent is moved.
+		RuntimeVirtualTextureComponent->TransformUpdated.AddUObject(this, &UVirtualHeightfieldMeshComponent::OnVirtualTextureTransformUpdate);
 		// Bind to delegate so that RuntimeVirtualTextureComponent will pull hide flags from this object.
 		RuntimeVirtualTextureComponent->GetHidePrimitivesDelegate().AddUObject(this, &UVirtualHeightfieldMeshComponent::GatherHideFlags);
 		RuntimeVirtualTextureComponent->MarkRenderStateDirty();
@@ -42,6 +44,7 @@ void UVirtualHeightfieldMeshComponent::OnUnregister()
 	URuntimeVirtualTextureComponent* RuntimeVirtualTextureComponent = VirtualTextureRef != nullptr ? VirtualTextureRef->VirtualTextureComponent : nullptr;
 	if (RuntimeVirtualTextureComponent)
 	{
+		RuntimeVirtualTextureComponent->TransformUpdated.RemoveAll(this);
 		RuntimeVirtualTextureComponent->GetHidePrimitivesDelegate().RemoveAll(this);
 		RuntimeVirtualTextureComponent->MarkRenderStateDirty();
 	}
@@ -105,6 +108,11 @@ void UVirtualHeightfieldMeshComponent::GatherHideFlags(bool& InOutHidePrimitives
 	InOutHidePrimitivesInGame |= bIsEnabled;
 }
 
+void UVirtualHeightfieldMeshComponent::OnVirtualTextureTransformUpdate(USceneComponent* InRootComponent, EUpdateTransformFlags UpdateTransformFlags, ETeleportType Teleport)
+{
+	MarkRenderStateDirty();
+}
+
 #if WITH_EDITOR
 
 void UVirtualHeightfieldMeshComponent::PostEditChangeProperty(FPropertyChangedEvent& PropertyChangedEvent)
diff --git a/Engine/Plugins/Experimental/VirtualHeightfieldMesh/Source/VirtualHeightfieldMesh/Public/VirtualHeightfieldMeshComponent.h b/Engine/Plugins/Experimental/VirtualHeightfieldMesh/Source/VirtualHeightfieldMesh/Public/VirtualHeightfieldMeshComponent.h
--- a/Engine/Plugins/Experimental/VirtualHeightfieldMesh/Source/VirtualHeightfieldMesh/Public/VirtualHeightfieldMeshComponent.h
+++ b/Engine/Plugins/Experimental/VirtualHeightfieldMesh/Source/VirtualHeightfieldMesh/Public/VirtualHeightfieldMeshComponent.h
@@ -116,6 +116,9 @@ protected:
 	UFUNCTION()
 	void GatherHideFlags(bool& InOutHidePrimitivesInEditor, bool& InOutHidePrimitivesInGame) const;
 
+	/** Handles when the transform is updated on the associated VirtualTexture. */
+	void OnVirtualTextureTransformUpdate(USceneComponent* InRootComponent, EUpdateTransformFlags UpdateTransformFlags, ETeleportType Teleport);
+
 	//~ Begin UObject Interface.
 #if WITH_EDITOR
 	virtual void PostEditChangeProperty(FPropertyChangedEvent& PropertyChangedEvent) override;
