Prime feedback buffer for VHM.
Prevents too much lod popping when spinning the camera.

diff --git a/Engine/Plugins/Experimental/VirtualHeightfieldMesh/Shaders/Private/VirtualHeightfieldMesh.usf b/Engine/Plugins/Experimental/VirtualHeightfieldMesh/Shaders/Private/VirtualHeightfieldMesh.usf
--- a/Engine/Plugins/Experimental/VirtualHeightfieldMesh/Shaders/Private/VirtualHeightfieldMesh.usf
+++ b/Engine/Plugins/Experimental/VirtualHeightfieldMesh/Shaders/Private/VirtualHeightfieldMesh.usf
@@ -38,6 +38,7 @@ int OcclusionLevelOffset;
 Texture2D<float> LodTexture;
 
 uint MaxLevel;
+uint NumForceLoadLods;
 uint PageTableFeedbackId;
 uint NumPhysicalAddressBits;
 float4 PageTableSize;
@@ -194,8 +195,23 @@ void InitBuffersCS()
 	RWIndirectArgsBuffer[7] = 1;
 	RWIndirectArgsBuffer[8] = 0;
 
-	// Clear virtual texture feedback counter.
-	RWFeedbackBuffer[0] = 0;
+	// Clear virtual texture feedback counter or prime feedback buffer with first n levels.
+	uint FeedbackCount = 0;
+
+	for (uint Level = 0; Level < NumForceLoadLods; ++Level)
+	{
+		for (int Y = 0; Y < (1 << Level); ++Y)
+		{
+			for (int X = 0; X < (1 << Level); ++X)
+			{
+				uint LevelPlusOne = 1 + MaxLevel - Level;
+				uint Feedback = X | (Y << 12) | (LevelPlusOne << 24) | PageTableFeedbackId;
+				RWFeedbackBuffer[++FeedbackCount] = Feedback;
+			}
+		}
+	}
+
+	RWFeedbackBuffer[0] = FeedbackCount;
 }
 
 /**
diff --git a/Engine/Plugins/Experimental/VirtualHeightfieldMesh/Source/VirtualHeightfieldMesh/Private/VirtualHeightfieldMeshSceneProxy.cpp b/Engine/Plugins/Experimental/VirtualHeightfieldMesh/Source/VirtualHeightfieldMesh/Private/VirtualHeightfieldMeshSceneProxy.cpp
--- a/Engine/Plugins/Experimental/VirtualHeightfieldMesh/Source/VirtualHeightfieldMesh/Private/VirtualHeightfieldMeshSceneProxy.cpp
+++ b/Engine/Plugins/Experimental/VirtualHeightfieldMesh/Source/VirtualHeightfieldMesh/Private/VirtualHeightfieldMeshSceneProxy.cpp
@@ -371,6 +371,7 @@ FVirtualHeightfieldMeshSceneProxy::FVirtualHeightfieldMeshSceneProxy(UVirtualHei
 	, LodDistribution(InComponent->GetLodDistribution())
 	, NumSubdivisionLODs(InComponent->GetNumSubdivisionLods())
 	, NumTailLods(InComponent->GetNumTailLods())
+	, NumForceLoadLods(InComponent->GetNumForceLoadLods())
 	, NumOcclusionLods(0)
 	, OcclusionGridSize(0, 0)
 {
@@ -688,6 +689,8 @@ namespace VirtualHeightfieldMesh
 
 		BEGIN_SHADER_PARAMETER_STRUCT(FParameters, )
 			SHADER_PARAMETER(uint32, MaxLevel)
+			SHADER_PARAMETER(uint32, NumForceLoadLods)
+			SHADER_PARAMETER(uint32, PageTableFeedbackId)
 			SHADER_PARAMETER_RDG_BUFFER_UAV(RWStructuredBuffer<WorkerQueueInfo>, RWQueueInfo)
 			SHADER_PARAMETER_RDG_BUFFER_UAV(RWBuffer<uint>, RWQueueBuffer)
 			SHADER_PARAMETER_RDG_BUFFER_UAV(RWBuffer<uint2>, RWQuadBuffer)
@@ -982,6 +985,7 @@ namespace VirtualHeightfieldMesh
 		int32 MinMaxLevelOffset;
 
 		uint32 MaxLevel;
+		uint32 NumForceLoadLods;
 		uint32 PageTableFeedbackId;
 		uint32 NumPhysicalAddressBits;
 		FVector4 PageTableSize;
@@ -1137,6 +1141,8 @@ namespace VirtualHeightfieldMesh
 
 		FInitBuffersCS::FParameters* PassParameters = GraphBuilder.AllocParameters<FInitBuffersCS::FParameters>();
 		PassParameters->MaxLevel = InDesc.MaxLevel;
+		PassParameters->NumForceLoadLods = InDesc.NumForceLoadLods;
+		PassParameters->PageTableFeedbackId = InDesc.PageTableFeedbackId;
 		PassParameters->RWQueueInfo = InVolatileResources.QueueInfoUAV;
 		PassParameters->RWQueueBuffer = InVolatileResources.QueueBufferUAV;
 		PassParameters->RWQuadBuffer = InVolatileResources.QuadBufferUAV;
@@ -1383,6 +1389,7 @@ void FVirtualHeightfieldMeshRendererExtension::SubmitWork(FRDGBuilder& GraphBuil
 		ProxyDesc.MinMaxTexture = Proxy->MinMaxTexture ? Proxy->MinMaxTexture->Resource->TextureRHI : VirtualHeightfieldMesh::GMinMaxDefaultTexture->TextureRHI;
 		ProxyDesc.MinMaxLevelOffset = ProxyDesc.MinMaxTexture->GetNumMips() - 1 - AllocatedVirtualTexture->GetMaxLevel();
 		ProxyDesc.MaxLevel = AllocatedVirtualTexture->GetMaxLevel();
+		ProxyDesc.NumForceLoadLods = FMath::Min<uint32>(FMath::Max(Proxy->NumForceLoadLods, 0), ProxyDesc.MaxLevel + 1);
 		ProxyDesc.PageTableSize = PageTableSize;
 		ProxyDesc.PhysicalPageTransform = PhysicalPageTransform;
 		ProxyDesc.NumPhysicalAddressBits = AllocatedVirtualTexture->GetPageTableFormat() == EVTPageTableFormat::UInt16 ? 6 : 8; // See packing in PageTableUpdate.usf
diff --git a/Engine/Plugins/Experimental/VirtualHeightfieldMesh/Source/VirtualHeightfieldMesh/Private/VirtualHeightfieldMeshSceneProxy.h b/Engine/Plugins/Experimental/VirtualHeightfieldMesh/Source/VirtualHeightfieldMesh/Private/VirtualHeightfieldMeshSceneProxy.h
--- a/Engine/Plugins/Experimental/VirtualHeightfieldMesh/Source/VirtualHeightfieldMesh/Private/VirtualHeightfieldMeshSceneProxy.h
+++ b/Engine/Plugins/Experimental/VirtualHeightfieldMesh/Source/VirtualHeightfieldMesh/Private/VirtualHeightfieldMeshSceneProxy.h
@@ -55,6 +55,7 @@ public:
 	
 	int32 NumSubdivisionLODs;
 	int32 NumTailLods;
+	int32 NumForceLoadLods;
 
 	TArray<FVector2D> OcclusionData;
 	int32 NumOcclusionLods;
diff --git a/Engine/Plugins/Experimental/VirtualHeightfieldMesh/Source/VirtualHeightfieldMesh/Public/VirtualHeightfieldMeshComponent.h b/Engine/Plugins/Experimental/VirtualHeightfieldMesh/Source/VirtualHeightfieldMesh/Public/VirtualHeightfieldMeshComponent.h
--- a/Engine/Plugins/Experimental/VirtualHeightfieldMesh/Source/VirtualHeightfieldMesh/Public/VirtualHeightfieldMeshComponent.h
+++ b/Engine/Plugins/Experimental/VirtualHeightfieldMesh/Source/VirtualHeightfieldMesh/Public/VirtualHeightfieldMeshComponent.h
@@ -70,6 +70,10 @@ protected:
 	UPROPERTY(EditAnywhere, Category = Rendering, meta = (DisplayName = "Tail LODs", ClampMin = "0", ClampMax = "8", UIMin = "0", UIMax = "8"))
 	int32 NumTailLods = 0;
 
+	/** The number of levels that we force loaded. A higher number can reduce LOD pop in but allocates more of the virtual texture memory pool. */
+	UPROPERTY(EditAnywhere, Category = Rendering, meta = (DisplayName = "Force Load LODs", ClampMin = "0", ClampMax = "4"))
+	int32 NumForceLoadLods = 0;
+
 	/** The number of levels that we calculate occlusion volumes for. A higher number gives finer occlusion at the cost of more queries. */
 	UPROPERTY(EditAnywhere, Category = Rendering, meta = (DisplayName = "Occlusion LODs", ClampMin = "0", ClampMax = "5"))
 	int32 NumOcclusionLods = 0;
@@ -109,6 +113,7 @@ public:
 	float GetLodDistribution() const { return LodDistribution; }
 	int32 GetNumSubdivisionLods() const { return NumSubdivisionLods; }
 	int32 GetNumTailLods() const { return NumTailLods; }
+	int32 GetNumForceLoadLods() const { return NumForceLoadLods; }
 	int32 GetNumOcclusionLods() const { return NumOcclusionLods; }
 
 protected:
