Bulk pass to address PVS warnings due to use of a TObjectPtr in the result clauses of a ternary operator.

diff --git a/Engine/Plugins/Developer/AnimationSharing/Source/AnimationSharing/Private/AnimationSharingManager.cpp b/Engine/Plugins/Developer/AnimationSharing/Source/AnimationSharing/Private/AnimationSharingManager.cpp
--- a/Engine/Plugins/Developer/AnimationSharing/Source/AnimationSharing/Private/AnimationSharingManager.cpp
+++ b/Engine/Plugins/Developer/AnimationSharing/Source/AnimationSharing/Private/AnimationSharingManager.cpp
@@ -830,7 +830,7 @@ void UAnimSharingInstance::SetupState(FPerStateData& StateData, const FAnimation
 	/** Setup overall data and flags */
 	StateData.bIsOnDemand = StateEntry.bOnDemand;
 	StateData.bIsAdditive = StateEntry.bAdditive;
-	StateData.AdditiveAnimationSequence = (StateEntry.bAdditive && StateEntry.AnimationSetups.IsValidIndex(0)) ? StateEntry.AnimationSetups[0].AnimSequence : nullptr;
+	StateData.AdditiveAnimationSequence = (StateEntry.bAdditive && StateEntry.AnimationSetups.IsValidIndex(0)) ? ToRawPtr(StateEntry.AnimationSetups[0].AnimSequence) : nullptr;
 
 	/** Keep hard reference to animation sequence */
 	if (StateData.AdditiveAnimationSequence)
diff --git a/Engine/Plugins/Experimental/Animation/MotionWarping/Source/MotionWarping/Private/MotionWarpingComponent.cpp b/Engine/Plugins/Experimental/Animation/MotionWarping/Source/MotionWarping/Private/MotionWarpingComponent.cpp
--- a/Engine/Plugins/Experimental/Animation/MotionWarping/Source/MotionWarping/Private/MotionWarpingComponent.cpp
+++ b/Engine/Plugins/Experimental/Animation/MotionWarping/Source/MotionWarping/Private/MotionWarpingComponent.cpp
@@ -137,7 +137,7 @@ void UMotionWarpingComponent::Update()
 	check(GetCharacterOwner());
 
 	const FAnimMontageInstance* RootMotionMontageInstance = GetCharacterOwner()->GetRootMotionAnimMontageInstance();
-	UAnimMontage* Montage = RootMotionMontageInstance ? RootMotionMontageInstance->Montage : nullptr;
+	UAnimMontage* Montage = RootMotionMontageInstance ? ToRawPtr(RootMotionMontageInstance->Montage) : nullptr;
 	if (Montage)
 	{
 		const float PreviousPosition = RootMotionMontageInstance->GetPreviousPosition();
diff --git a/Engine/Plugins/Experimental/Animation/MotionWarping/Source/MotionWarping/Private/RootMotionModifier.cpp b/Engine/Plugins/Experimental/Animation/MotionWarping/Source/MotionWarping/Private/RootMotionModifier.cpp
--- a/Engine/Plugins/Experimental/Animation/MotionWarping/Source/MotionWarping/Private/RootMotionModifier.cpp
+++ b/Engine/Plugins/Experimental/Animation/MotionWarping/Source/MotionWarping/Private/RootMotionModifier.cpp
@@ -14,7 +14,7 @@
 void FRootMotionModifier::Update(UMotionWarpingComponent& OwnerComp)
 {
 	const FAnimMontageInstance* RootMotionMontageInstance = OwnerComp.GetCharacterOwner()->GetRootMotionAnimMontageInstance();
-	const UAnimMontage* Montage = RootMotionMontageInstance ? RootMotionMontageInstance->Montage : nullptr;
+	const UAnimMontage* Montage = RootMotionMontageInstance ? ToRawPtr(RootMotionMontageInstance->Montage) : nullptr;
 
 	// Mark for removal if our animation is not relevant anymore
 	if (Montage == nullptr || Montage != Animation)
diff --git a/Engine/Plugins/Experimental/Animation/MotionWarping/Source/MotionWarping/Private/RootMotionModifier_AdjustmentBlendWarp.cpp b/Engine/Plugins/Experimental/Animation/MotionWarping/Source/MotionWarping/Private/RootMotionModifier_AdjustmentBlendWarp.cpp
--- a/Engine/Plugins/Experimental/Animation/MotionWarping/Source/MotionWarping/Private/RootMotionModifier_AdjustmentBlendWarp.cpp
+++ b/Engine/Plugins/Experimental/Animation/MotionWarping/Source/MotionWarping/Private/RootMotionModifier_AdjustmentBlendWarp.cpp
@@ -412,7 +412,7 @@ void URootMotionModifierConfig_AdjustmentBlendWarp::GetIKBoneTransformAndAlpha(A
 				// We must check if the Animation for the modifier is still relevant because in RootMotionFromMontageOnlyMode the montage could be aborted 
 				// but the modifier will remain in the list until the next update
 				const FAnimMontageInstance* RootMotionMontageInstance = Character->GetRootMotionAnimMontageInstance();
-				const UAnimMontage* Montage = RootMotionMontageInstance ? RootMotionMontageInstance->Montage : nullptr;
+				const UAnimMontage* Montage = RootMotionMontageInstance ? ToRawPtr(RootMotionMontageInstance->Montage) : nullptr;
 				if (Modifier->Animation == Montage)
 				{
 					static_cast<const FRootMotionModifier_AdjustmentBlendWarp*>(Modifier.Get())->GetIKBoneTransformAndAlpha(BoneName, OutTransform, OutAlpha);
diff --git a/Engine/Plugins/Experimental/VirtualHeightfieldMesh/Source/VirtualHeightfieldMesh/Private/VirtualHeightfieldMeshComponent.cpp b/Engine/Plugins/Experimental/VirtualHeightfieldMesh/Source/VirtualHeightfieldMesh/Private/VirtualHeightfieldMeshComponent.cpp
--- a/Engine/Plugins/Experimental/VirtualHeightfieldMesh/Source/VirtualHeightfieldMesh/Private/VirtualHeightfieldMeshComponent.cpp
+++ b/Engine/Plugins/Experimental/VirtualHeightfieldMesh/Source/VirtualHeightfieldMesh/Private/VirtualHeightfieldMeshComponent.cpp
@@ -26,7 +26,7 @@ void UVirtualHeightfieldMeshComponent::OnRegister()
 {
 	VirtualTextureRef = VirtualTexture.Get();
 
-	URuntimeVirtualTextureComponent* RuntimeVirtualTextureComponent = VirtualTextureRef != nullptr ? VirtualTextureRef->VirtualTextureComponent : nullptr;
+	URuntimeVirtualTextureComponent* RuntimeVirtualTextureComponent = VirtualTextureRef != nullptr ? ToRawPtr(VirtualTextureRef->VirtualTextureComponent) : nullptr;
 	if (RuntimeVirtualTextureComponent)
 	{
 		// Bind to delegate so that we dirty render state whenever RuntimeVirtualTextureComponent is moved.
@@ -41,7 +41,7 @@ void UVirtualHeightfieldMeshComponent::OnRegister()
 
 void UVirtualHeightfieldMeshComponent::OnUnregister()
 {
-	URuntimeVirtualTextureComponent* RuntimeVirtualTextureComponent = VirtualTextureRef != nullptr ? VirtualTextureRef->VirtualTextureComponent : nullptr;
+	URuntimeVirtualTextureComponent* RuntimeVirtualTextureComponent = VirtualTextureRef != nullptr ? ToRawPtr(VirtualTextureRef->VirtualTextureComponent) : nullptr;
 	if (RuntimeVirtualTextureComponent)
 	{
 		RuntimeVirtualTextureComponent->TransformUpdated.RemoveAll(this);
@@ -61,13 +61,13 @@ ARuntimeVirtualTextureVolume* UVirtualHeightfieldMeshComponent::GetVirtualTextur
 
 URuntimeVirtualTexture* UVirtualHeightfieldMeshComponent::GetVirtualTexture() const
 {
-	URuntimeVirtualTextureComponent* RuntimeVirtualTextureComponent = VirtualTextureRef != nullptr ? VirtualTextureRef->VirtualTextureComponent : nullptr;
+	URuntimeVirtualTextureComponent* RuntimeVirtualTextureComponent = VirtualTextureRef != nullptr ? ToRawPtr(VirtualTextureRef->VirtualTextureComponent) : nullptr;
 	return RuntimeVirtualTextureComponent ? RuntimeVirtualTextureComponent->GetVirtualTexture() : nullptr;
 }
 
 FTransform UVirtualHeightfieldMeshComponent::GetVirtualTextureTransform() const
 {
-	URuntimeVirtualTextureComponent* RuntimeVirtualTextureComponent = VirtualTextureRef != nullptr ? VirtualTextureRef->VirtualTextureComponent : nullptr;
+	URuntimeVirtualTextureComponent* RuntimeVirtualTextureComponent = VirtualTextureRef != nullptr ? ToRawPtr(VirtualTextureRef->VirtualTextureComponent) : nullptr;
 	return RuntimeVirtualTextureComponent ? RuntimeVirtualTextureComponent->GetComponentTransform() * RuntimeVirtualTextureComponent->GetTexelSnapTransform() : FTransform::Identity;
 }
 
@@ -123,7 +123,7 @@ void UVirtualHeightfieldMeshComponent::PostEditChangeProperty(FPropertyChangedEv
 	if (PropertyName == HideInEditorName)
 	{
 		// Force RuntimeVirtualTextureComponent to poll the HidePrimitives settings.
-		URuntimeVirtualTextureComponent* RuntimeVirtualTextureComponent = VirtualTextureRef != nullptr ? VirtualTextureRef->VirtualTextureComponent : nullptr;
+		URuntimeVirtualTextureComponent* RuntimeVirtualTextureComponent = VirtualTextureRef != nullptr ? ToRawPtr(VirtualTextureRef->VirtualTextureComponent) : nullptr;
 		if (RuntimeVirtualTextureComponent != nullptr)
 		{
 			RuntimeVirtualTextureComponent->MarkRenderStateDirty();
diff --git a/Engine/Plugins/Experimental/VirtualHeightfieldMesh/Source/VirtualHeightfieldMesh/Private/VirtualHeightfieldMeshEnable.cpp b/Engine/Plugins/Experimental/VirtualHeightfieldMesh/Source/VirtualHeightfieldMesh/Private/VirtualHeightfieldMeshEnable.cpp
--- a/Engine/Plugins/Experimental/VirtualHeightfieldMesh/Source/VirtualHeightfieldMesh/Private/VirtualHeightfieldMeshEnable.cpp
+++ b/Engine/Plugins/Experimental/VirtualHeightfieldMesh/Source/VirtualHeightfieldMesh/Private/VirtualHeightfieldMeshEnable.cpp
@@ -37,7 +37,7 @@ namespace VirtualHeightfieldMesh
 				It->MarkRenderStateDirty();
 
 				ARuntimeVirtualTextureVolume* VirtualTextureVolume = It->GetVirtualTextureVolume();
-				URuntimeVirtualTextureComponent* VirtualTextureComponent = VirtualTextureVolume != nullptr ? VirtualTextureVolume->VirtualTextureComponent : nullptr;
+				URuntimeVirtualTextureComponent* VirtualTextureComponent = VirtualTextureVolume != nullptr ? ToRawPtr(VirtualTextureVolume->VirtualTextureComponent) : nullptr;
 				URuntimeVirtualTexture* VirtualTexture = VirtualTextureComponent != nullptr ? VirtualTextureComponent->GetVirtualTexture() : nullptr;
 
 				if (VirtualTextureComponent != nullptr)
diff --git a/Engine/Plugins/Experimental/VirtualHeightfieldMesh/Source/VirtualHeightfieldMeshEditor/Private/HeightfieldMinMaxTextureBuild.cpp b/Engine/Plugins/Experimental/VirtualHeightfieldMesh/Source/VirtualHeightfieldMeshEditor/Private/HeightfieldMinMaxTextureBuild.cpp
--- a/Engine/Plugins/Experimental/VirtualHeightfieldMesh/Source/VirtualHeightfieldMeshEditor/Private/HeightfieldMinMaxTextureBuild.cpp
+++ b/Engine/Plugins/Experimental/VirtualHeightfieldMesh/Source/VirtualHeightfieldMeshEditor/Private/HeightfieldMinMaxTextureBuild.cpp
@@ -120,7 +120,7 @@ namespace VirtualHeightfieldMesh
 		}
 
 		ARuntimeVirtualTextureVolume* VirtualTextureVolume = InComponent->GetVirtualTextureVolume();
-		URuntimeVirtualTextureComponent* VirtualTextureComponent = VirtualTextureVolume != nullptr ? VirtualTextureVolume->VirtualTextureComponent : nullptr;
+		URuntimeVirtualTextureComponent* VirtualTextureComponent = VirtualTextureVolume != nullptr ? ToRawPtr(VirtualTextureVolume->VirtualTextureComponent) : nullptr;
 
 		if (VirtualTextureComponent == nullptr)
 		{
diff --git a/Engine/Plugins/NetcodeUnitTest/NetcodeUnitTest/Source/NetcodeUnitTest/Private/MinimalClient.cpp b/Engine/Plugins/NetcodeUnitTest/NetcodeUnitTest/Source/NetcodeUnitTest/Private/MinimalClient.cpp
--- a/Engine/Plugins/NetcodeUnitTest/NetcodeUnitTest/Source/NetcodeUnitTest/Private/MinimalClient.cpp
+++ b/Engine/Plugins/NetcodeUnitTest/NetcodeUnitTest/Source/NetcodeUnitTest/Private/MinimalClient.cpp
@@ -585,7 +585,7 @@ bool UMinimalClient::PostSendRPC(FString RPCName, UObject* Target/*=nullptr*/)
 
 void UMinimalClient::ResetConnTimeout(float Duration)
 {
-	UNetDriver* UnitDriver = (UnitConn != nullptr ? UnitConn->Driver : nullptr);
+	UNetDriver* UnitDriver = (UnitConn != nullptr ? ToRawPtr(UnitConn->Driver) : nullptr);
 
 	if (UnitDriver != nullptr && UnitConn->State != USOCK_Closed)
 	{
diff --git a/Engine/Plugins/NetcodeUnitTest/NetcodeUnitTest/Source/NetcodeUnitTest/Private/NUTActor.cpp b/Engine/Plugins/NetcodeUnitTest/NetcodeUnitTest/Source/NetcodeUnitTest/Private/NUTActor.cpp
--- a/Engine/Plugins/NetcodeUnitTest/NetcodeUnitTest/Source/NetcodeUnitTest/Private/NUTActor.cpp
+++ b/Engine/Plugins/NetcodeUnitTest/NetcodeUnitTest/Source/NetcodeUnitTest/Private/NUTActor.cpp
@@ -72,7 +72,7 @@ UNetConnection* ANUTActor::GetNetConnection() const
 	if (ReturnVal == NULL)
 	{
 		UNetDriver* NetDriver = GetNetDriver();
-		UNetConnection* ServerConn = (NetDriver != NULL ? NetDriver->ServerConnection : NULL);
+		UNetConnection* ServerConn = (NetDriver != NULL ? ToRawPtr(NetDriver->ServerConnection) : NULL);
 
 		if (ServerConn != NULL)
 		{
diff --git a/Engine/Plugins/NetcodeUnitTest/NetcodeUnitTest/Source/NetcodeUnitTest/Private/Net/NUTUtilNet.cpp b/Engine/Plugins/NetcodeUnitTest/NetcodeUnitTest/Source/NetcodeUnitTest/Private/Net/NUTUtilNet.cpp
--- a/Engine/Plugins/NetcodeUnitTest/NetcodeUnitTest/Source/NetcodeUnitTest/Private/Net/NUTUtilNet.cpp
+++ b/Engine/Plugins/NetcodeUnitTest/NetcodeUnitTest/Source/NetcodeUnitTest/Private/Net/NUTUtilNet.cpp
@@ -529,8 +529,8 @@ UWorld* NUTNet::CreateUnitTestWorld(bool bHookTick/*=true*/)
 			if (CurSettings != NULL)
 			{
 				ULocalPlayer* PrimLocPlayer = GEngine->GetFirstGamePlayer(NUTUtil::GetPrimaryWorld());
-				APlayerController* PrimPC = (PrimLocPlayer != NULL ? PrimLocPlayer->PlayerController : NULL);
-				APlayerState* PrimState = (PrimPC != NULL ? PrimPC->PlayerState : NULL);
+				APlayerController* PrimPC = (PrimLocPlayer != NULL ? ToRawPtr(PrimLocPlayer->PlayerController) : NULL);
+				APlayerState* PrimState = (PrimPC != NULL ? ToRawPtr(PrimPC->PlayerState) : NULL);
 
 				if (PrimState != NULL)
 				{
diff --git a/Engine/Plugins/NetcodeUnitTest/NetcodeUnitTest/Source/NetcodeUnitTest/Private/UnitTestManager.cpp b/Engine/Plugins/NetcodeUnitTest/NetcodeUnitTest/Source/NetcodeUnitTest/Private/UnitTestManager.cpp
--- a/Engine/Plugins/NetcodeUnitTest/NetcodeUnitTest/Source/NetcodeUnitTest/Private/UnitTestManager.cpp
+++ b/Engine/Plugins/NetcodeUnitTest/NetcodeUnitTest/Source/NetcodeUnitTest/Private/UnitTestManager.cpp
@@ -1997,7 +1997,7 @@ void UUnitTestManager::Serialize(const TCHAR* Data, ELogVerbosity::Type Verbosit
 			{
 				// @todo #JohnB: This chain of getting the desired world is very messy...but it's hard to get the correct world.
 				UGameEngine* GameEngine = Cast<UGameEngine>(GEngine);
-				UGameInstance* GameInstance = GameEngine != nullptr ? GameEngine->GameInstance : nullptr;
+				UGameInstance* GameInstance = GameEngine != nullptr ? ToRawPtr(GameEngine->GameInstance) : nullptr;
 				UGameViewportClient* GameViewport = GameInstance != nullptr ? GameInstance->GetGameViewportClient() : nullptr;
 				UWorld* BaseWorld = GameViewport != nullptr ? GameViewport->GetWorld() : nullptr;
 
diff --git a/Engine/Plugins/Online/OnlineSubsystemUtils/Source/OnlineSubsystemUtils/Private/InAppPurchaseCallbackProxy.cpp b/Engine/Plugins/Online/OnlineSubsystemUtils/Source/OnlineSubsystemUtils/Private/InAppPurchaseCallbackProxy.cpp
--- a/Engine/Plugins/Online/OnlineSubsystemUtils/Source/OnlineSubsystemUtils/Private/InAppPurchaseCallbackProxy.cpp
+++ b/Engine/Plugins/Online/OnlineSubsystemUtils/Source/OnlineSubsystemUtils/Private/InAppPurchaseCallbackProxy.cpp
@@ -24,7 +24,7 @@ void UInAppPurchaseCallbackProxy::Trigger(APlayerController* PlayerController, c
 	EInAppPurchaseState::Type TempState = EInAppPurchaseState::Unknown;
 
 	WorldPtr = (PlayerController != nullptr) ? PlayerController->GetWorld() : nullptr;
-	if (APlayerState* PlayerState = (PlayerController != nullptr) ? PlayerController->PlayerState : nullptr)
+	if (APlayerState* PlayerState = (PlayerController != nullptr) ? ToRawPtr(PlayerController->PlayerState) : nullptr)
 	{
 		if (IOnlineSubsystem* const OnlineSub = IOnlineSubsystem::IsLoaded() ? IOnlineSubsystem::Get() : nullptr)
 		{
diff --git a/Engine/Plugins/Online/OnlineSubsystemUtils/Source/OnlineSubsystemUtils/Private/InAppPurchaseCallbackProxy2.cpp b/Engine/Plugins/Online/OnlineSubsystemUtils/Source/OnlineSubsystemUtils/Private/InAppPurchaseCallbackProxy2.cpp
--- a/Engine/Plugins/Online/OnlineSubsystemUtils/Source/OnlineSubsystemUtils/Private/InAppPurchaseCallbackProxy2.cpp
+++ b/Engine/Plugins/Online/OnlineSubsystemUtils/Source/OnlineSubsystemUtils/Private/InAppPurchaseCallbackProxy2.cpp
@@ -23,7 +23,7 @@ void UInAppPurchaseCallbackProxy2::Trigger(APlayerController* PlayerController,
 {
 	bFailedToEvenSubmit = true;
 	WorldPtr = (PlayerController != nullptr) ? PlayerController->GetWorld() : nullptr;
-	if (APlayerState* PlayerState = (PlayerController != nullptr) ? PlayerController->PlayerState : nullptr)
+	if (APlayerState* PlayerState = (PlayerController != nullptr) ? ToRawPtr(PlayerController->PlayerState) : nullptr)
 	{
 		if (IOnlineSubsystem* const OnlineSub = IOnlineSubsystem::IsLoaded() ? IOnlineSubsystem::Get() : nullptr)
 		{
@@ -65,7 +65,7 @@ void UInAppPurchaseCallbackProxy2::Trigger(APlayerController* PlayerController,
 void UInAppPurchaseCallbackProxy2::TriggerGetUnprocessedPurchases(APlayerController* PlayerController)
 {
 	WorldPtr = (PlayerController != nullptr) ? PlayerController->GetWorld() : nullptr;
-	if (APlayerState* PlayerState = (PlayerController != nullptr) ? PlayerController->PlayerState : nullptr)
+	if (APlayerState* PlayerState = (PlayerController != nullptr) ? ToRawPtr(PlayerController->PlayerState) : nullptr)
 	{
 		if (IOnlineSubsystem* const OnlineSub = IOnlineSubsystem::IsLoaded() ? IOnlineSubsystem::Get() : nullptr)
 		{
@@ -95,7 +95,7 @@ void UInAppPurchaseCallbackProxy2::TriggerGetOwnedPurchases(APlayerController* P
 {
 	bFailedToEvenSubmit = true;
 	WorldPtr = (PlayerController != nullptr) ? PlayerController->GetWorld() : nullptr;
-	if (APlayerState* PlayerState = (PlayerController != nullptr) ? PlayerController->PlayerState : nullptr)
+	if (APlayerState* PlayerState = (PlayerController != nullptr) ? ToRawPtr(PlayerController->PlayerState) : nullptr)
 	{
 		if (IOnlineSubsystem* const OnlineSub = IOnlineSubsystem::IsLoaded() ? IOnlineSubsystem::Get() : nullptr)
 		{
diff --git a/Engine/Plugins/Online/OnlineSubsystemUtils/Source/OnlineSubsystemUtils/Private/InAppPurchaseQueryCallbackProxy.cpp b/Engine/Plugins/Online/OnlineSubsystemUtils/Source/OnlineSubsystemUtils/Private/InAppPurchaseQueryCallbackProxy.cpp
--- a/Engine/Plugins/Online/OnlineSubsystemUtils/Source/OnlineSubsystemUtils/Private/InAppPurchaseQueryCallbackProxy.cpp
+++ b/Engine/Plugins/Online/OnlineSubsystemUtils/Source/OnlineSubsystemUtils/Private/InAppPurchaseQueryCallbackProxy.cpp
@@ -19,7 +19,7 @@ void UInAppPurchaseQueryCallbackProxy::TriggerQuery(APlayerController* PlayerCon
 	bFailedToEvenSubmit = true;
 
 	WorldPtr = (PlayerController != nullptr) ? PlayerController->GetWorld() : nullptr;
-	if (APlayerState* PlayerState = (PlayerController != NULL) ? PlayerController->PlayerState : nullptr)
+	if (APlayerState* PlayerState = (PlayerController != NULL) ? ToRawPtr(PlayerController->PlayerState) : nullptr)
 	{
 		if (IOnlineSubsystem* const OnlineSub = IOnlineSubsystem::IsLoaded() ? IOnlineSubsystem::Get() : nullptr)
 		{
diff --git a/Engine/Plugins/Online/OnlineSubsystemUtils/Source/OnlineSubsystemUtils/Private/InAppPurchaseQueryCallbackProxy2.cpp b/Engine/Plugins/Online/OnlineSubsystemUtils/Source/OnlineSubsystemUtils/Private/InAppPurchaseQueryCallbackProxy2.cpp
--- a/Engine/Plugins/Online/OnlineSubsystemUtils/Source/OnlineSubsystemUtils/Private/InAppPurchaseQueryCallbackProxy2.cpp
+++ b/Engine/Plugins/Online/OnlineSubsystemUtils/Source/OnlineSubsystemUtils/Private/InAppPurchaseQueryCallbackProxy2.cpp
@@ -20,7 +20,7 @@ void UInAppPurchaseQueryCallbackProxy2::TriggerQuery(APlayerController* PlayerCo
 	bool bFailedToEvenSubmit = true;
 
 	WorldPtr = (PlayerController != nullptr) ? PlayerController->GetWorld() : nullptr;
-	if (APlayerState* PlayerState = (PlayerController != NULL) ? PlayerController->PlayerState : nullptr)
+	if (APlayerState* PlayerState = (PlayerController != NULL) ? ToRawPtr(PlayerController->PlayerState) : nullptr)
 	{
 		if (IOnlineSubsystem* const OnlineSub = IOnlineSubsystem::IsLoaded() ? IOnlineSubsystem::Get() : nullptr)
 		{
diff --git a/Engine/Plugins/Online/OnlineSubsystemUtils/Source/OnlineSubsystemUtils/Private/InAppPurchaseRestoreCallbackProxy.cpp b/Engine/Plugins/Online/OnlineSubsystemUtils/Source/OnlineSubsystemUtils/Private/InAppPurchaseRestoreCallbackProxy.cpp
--- a/Engine/Plugins/Online/OnlineSubsystemUtils/Source/OnlineSubsystemUtils/Private/InAppPurchaseRestoreCallbackProxy.cpp
+++ b/Engine/Plugins/Online/OnlineSubsystemUtils/Source/OnlineSubsystemUtils/Private/InAppPurchaseRestoreCallbackProxy.cpp
@@ -21,7 +21,7 @@ void UInAppPurchaseRestoreCallbackProxy::Trigger(const TArray<FInAppPurchaseProd
 	bFailedToEvenSubmit = true;
 
 	WorldPtr = (PlayerController != nullptr) ? PlayerController->GetWorld() : nullptr;
-	if (APlayerState* PlayerState = (PlayerController != nullptr) ? PlayerController->PlayerState : nullptr)
+	if (APlayerState* PlayerState = (PlayerController != nullptr) ? ToRawPtr(PlayerController->PlayerState) : nullptr)
 	{
 		if (IOnlineSubsystem* const OnlineSub = IOnlineSubsystem::IsLoaded() ? IOnlineSubsystem::Get() : nullptr)
 		{
diff --git a/Engine/Plugins/Online/OnlineSubsystemUtils/Source/OnlineSubsystemUtils/Private/InAppPurchaseRestoreCallbackProxy2.cpp b/Engine/Plugins/Online/OnlineSubsystemUtils/Source/OnlineSubsystemUtils/Private/InAppPurchaseRestoreCallbackProxy2.cpp
--- a/Engine/Plugins/Online/OnlineSubsystemUtils/Source/OnlineSubsystemUtils/Private/InAppPurchaseRestoreCallbackProxy2.cpp
+++ b/Engine/Plugins/Online/OnlineSubsystemUtils/Source/OnlineSubsystemUtils/Private/InAppPurchaseRestoreCallbackProxy2.cpp
@@ -23,7 +23,7 @@ void UInAppPurchaseRestoreCallbackProxy2::Trigger(const TArray<FInAppPurchasePro
 	bool bFailedToEvenSubmit = true;
 
 	WorldPtr = (PlayerController != nullptr) ? PlayerController->GetWorld() : nullptr;
-	if (APlayerState* PlayerState = (PlayerController != nullptr) ? PlayerController->PlayerState : nullptr)
+	if (APlayerState* PlayerState = (PlayerController != nullptr) ? ToRawPtr(PlayerController->PlayerState) : nullptr)
 	{
 		if (IOnlineSubsystem* const OnlineSub = IOnlineSubsystem::IsLoaded() ? IOnlineSubsystem::Get() : nullptr)
 		{
diff --git a/Engine/Plugins/Online/OnlineSubsystemUtils/Source/OnlineSubsystemUtils/Private/LeaderboardBlueprintLibrary.cpp b/Engine/Plugins/Online/OnlineSubsystemUtils/Source/OnlineSubsystemUtils/Private/LeaderboardBlueprintLibrary.cpp
--- a/Engine/Plugins/Online/OnlineSubsystemUtils/Source/OnlineSubsystemUtils/Private/LeaderboardBlueprintLibrary.cpp
+++ b/Engine/Plugins/Online/OnlineSubsystemUtils/Source/OnlineSubsystemUtils/Private/LeaderboardBlueprintLibrary.cpp
@@ -18,7 +18,7 @@ ULeaderboardBlueprintLibrary::ULeaderboardBlueprintLibrary(const FObjectInitiali
 
 bool ULeaderboardBlueprintLibrary::WriteLeaderboardObject(APlayerController* PlayerController, class FOnlineLeaderboardWrite& WriteObject)
 {
-	if (APlayerState* PlayerState = (PlayerController != NULL) ? PlayerController->PlayerState : NULL)
+	if (APlayerState* PlayerState = (PlayerController != NULL) ? ToRawPtr(PlayerController->PlayerState) : NULL)
 	{
 		TSharedPtr<const FUniqueNetId> UserId = PlayerState->GetUniqueId().GetUniqueNetId();
 		if (UserId.IsValid())
diff --git a/Engine/Plugins/Online/OnlineSubsystemUtils/Source/OnlineSubsystemUtils/Private/LeaderboardFlushCallbackProxy.cpp b/Engine/Plugins/Online/OnlineSubsystemUtils/Source/OnlineSubsystemUtils/Private/LeaderboardFlushCallbackProxy.cpp
--- a/Engine/Plugins/Online/OnlineSubsystemUtils/Source/OnlineSubsystemUtils/Private/LeaderboardFlushCallbackProxy.cpp
+++ b/Engine/Plugins/Online/OnlineSubsystemUtils/Source/OnlineSubsystemUtils/Private/LeaderboardFlushCallbackProxy.cpp
@@ -17,7 +17,7 @@ void ULeaderboardFlushCallbackProxy::TriggerFlush(APlayerController* PlayerContr
 {
 	bFailedToEvenSubmit = true;
 
-	if (APlayerState* PlayerState = (PlayerController != NULL) ? PlayerController->PlayerState : NULL)
+	if (APlayerState* PlayerState = (PlayerController != NULL) ? ToRawPtr(PlayerController->PlayerState) : NULL)
 	{
 		TSharedPtr<const FUniqueNetId> UserID = PlayerState->GetUniqueId().GetUniqueNetId();
 		if (UserID.IsValid())
diff --git a/Engine/Plugins/Online/OnlineSubsystemUtils/Source/OnlineSubsystemUtils/Private/LeaderboardQueryCallbackProxy.cpp b/Engine/Plugins/Online/OnlineSubsystemUtils/Source/OnlineSubsystemUtils/Private/LeaderboardQueryCallbackProxy.cpp
--- a/Engine/Plugins/Online/OnlineSubsystemUtils/Source/OnlineSubsystemUtils/Private/LeaderboardQueryCallbackProxy.cpp
+++ b/Engine/Plugins/Online/OnlineSubsystemUtils/Source/OnlineSubsystemUtils/Private/LeaderboardQueryCallbackProxy.cpp
@@ -20,7 +20,7 @@ void ULeaderboardQueryCallbackProxy::TriggerQuery(APlayerController* PlayerContr
 	bFailedToEvenSubmit = true;
 
 	WorldPtr = (PlayerController != NULL) ? PlayerController->GetWorld() : NULL;
-	if (APlayerState* PlayerState = (PlayerController != NULL) ? PlayerController->PlayerState : NULL)
+	if (APlayerState* PlayerState = (PlayerController != NULL) ? ToRawPtr(PlayerController->PlayerState) : NULL)
 	{
 		TSharedPtr<const FUniqueNetId> UserID = PlayerState->GetUniqueId().GetUniqueNetId();
 		if (UserID.IsValid())
diff --git a/Engine/Plugins/Online/OnlineSubsystemUtils/Source/OnlineSubsystemUtils/Private/OnlineSubsystemUtils.cpp b/Engine/Plugins/Online/OnlineSubsystemUtils/Source/OnlineSubsystemUtils/Private/OnlineSubsystemUtils.cpp
--- a/Engine/Plugins/Online/OnlineSubsystemUtils/Source/OnlineSubsystemUtils/Private/OnlineSubsystemUtils.cpp
+++ b/Engine/Plugins/Online/OnlineSubsystemUtils/Source/OnlineSubsystemUtils/Private/OnlineSubsystemUtils.cpp
@@ -809,7 +809,7 @@ void FOnlineSubsystemBPCallHelper::QueryIDFromPlayerController(APlayerController
 {
 	UserID.Reset();
 
-	if (APlayerState* PlayerState = (PlayerController != NULL) ? PlayerController->PlayerState : NULL)
+	if (APlayerState* PlayerState = (PlayerController != NULL) ? ToRawPtr(PlayerController->PlayerState) : NULL)
 	{
 		UserID = PlayerState->GetUniqueId().GetUniqueNetId();
 		if (!UserID.IsValid())
diff --git a/Engine/Plugins/VirtualProduction/Takes/Source/TakeRecorderSources/Private/TakeRecorderCameraCutSource.cpp b/Engine/Plugins/VirtualProduction/Takes/Source/TakeRecorderSources/Private/TakeRecorderCameraCutSource.cpp
--- a/Engine/Plugins/VirtualProduction/Takes/Source/TakeRecorderSources/Private/TakeRecorderCameraCutSource.cpp
+++ b/Engine/Plugins/VirtualProduction/Takes/Source/TakeRecorderSources/Private/TakeRecorderCameraCutSource.cpp
@@ -56,7 +56,7 @@ void UTakeRecorderCameraCutSource::TickRecording(const FQualifiedFrameTime& Curr
 	else 
 	{
 		APlayerController* PC = World->GetGameInstance() ? World->GetGameInstance()->GetFirstLocalPlayerController() : nullptr;
-		APlayerCameraManager* PlayerCameraManager = PC ? PC->PlayerCameraManager : nullptr;
+		APlayerCameraManager* PlayerCameraManager = PC ? ToRawPtr(PC->PlayerCameraManager) : nullptr;
 		if (PlayerCameraManager)
 		{
 			Target = PlayerCameraManager->ViewTarget.Target;
diff --git a/Engine/Plugins/VirtualProduction/Takes/Source/TakeTrackRecorders/Private/MovieScene3DTransformTrackRecorder.cpp b/Engine/Plugins/VirtualProduction/Takes/Source/TakeTrackRecorders/Private/MovieScene3DTransformTrackRecorder.cpp
--- a/Engine/Plugins/VirtualProduction/Takes/Source/TakeTrackRecorders/Private/MovieScene3DTransformTrackRecorder.cpp
+++ b/Engine/Plugins/VirtualProduction/Takes/Source/TakeTrackRecorders/Private/MovieScene3DTransformTrackRecorder.cpp
@@ -489,7 +489,7 @@ void UMovieScene3DTransformTrackRecorder::PostProcessAnimationData(UMovieSceneAn
 
 			}
 			// Search for the Root Bone in the Skeleton
-			const USkeletalMesh* SkeletalMesh = SkeletalMeshComponent->MasterPoseComponent != nullptr ? SkeletalMeshComponent->MasterPoseComponent->SkeletalMesh : SkeletalMeshComponent->SkeletalMesh;
+			const USkeletalMesh* SkeletalMesh = SkeletalMeshComponent->MasterPoseComponent != nullptr ? ToRawPtr(SkeletalMeshComponent->MasterPoseComponent->SkeletalMesh) : ToRawPtr(SkeletalMeshComponent->SkeletalMesh);
 			const UAnimSequence* AnimSequence = AnimTrackRecorder->GetAnimSequence();
 			if (AnimSequence && SkeletalMesh && AnimSequence->GetSkeleton())
 			{
diff --git a/Engine/Source/Developer/BlueprintCompilerCppBackend/Private/BlueprintCompilerCppBackendValueHelper.cpp b/Engine/Source/Developer/BlueprintCompilerCppBackend/Private/BlueprintCompilerCppBackendValueHelper.cpp
--- a/Engine/Source/Developer/BlueprintCompilerCppBackend/Private/BlueprintCompilerCppBackendValueHelper.cpp
+++ b/Engine/Source/Developer/BlueprintCompilerCppBackend/Private/BlueprintCompilerCppBackendValueHelper.cpp
@@ -1887,7 +1887,7 @@ void FEmitDefaultValueHelper::GenerateCustomDynamicClassInitialization(FEmitterL
 		{
 			// >>> This code should be removed, once UE-39168 is fixed
 			//TODO: it's an ugly workaround - template from DefaultSceneRootNode is unnecessarily cooked :(
-			UActorComponent* DefaultSceneRootComponentTemplate = SCS->GetDefaultSceneRootNode() ? SCS->GetDefaultSceneRootNode()->ComponentTemplate : nullptr;
+			UActorComponent* DefaultSceneRootComponentTemplate = SCS->GetDefaultSceneRootNode() ? ToRawPtr(SCS->GetDefaultSceneRootNode()->ComponentTemplate) : nullptr;
 			if (DefaultSceneRootComponentTemplate)
 			{
 				ActorComponentTempatesOwnedByClass.Add(DefaultSceneRootComponentTemplate);
diff --git a/Engine/Source/Editor/MaterialEditor/Private/MaterialEditor.cpp b/Engine/Source/Editor/MaterialEditor/Private/MaterialEditor.cpp
--- a/Engine/Source/Editor/MaterialEditor/Private/MaterialEditor.cpp
+++ b/Engine/Source/Editor/MaterialEditor/Private/MaterialEditor.cpp
@@ -5110,7 +5110,7 @@ void FMaterialEditor::JumpToExpression(UMaterialExpression* Expression)
 		else
 		{
 			UMaterialExpressionParameter* GraphExpression = Material->FindExpressionByGUID<UMaterialExpressionParameter>(Expression->GetParameterExpressionId());
-			ExpressionNode = GraphExpression ? GraphExpression->GraphNode : nullptr;
+			ExpressionNode = GraphExpression ? ToRawPtr(GraphExpression->GraphNode) : nullptr;
 		}
 	}
 	else if (UMaterialExpressionFunctionOutput* ExpressionOutput = Cast<UMaterialExpressionFunctionOutput>(Expression))
@@ -5122,7 +5122,7 @@ void FMaterialEditor::JumpToExpression(UMaterialExpression* Expression)
 			{
 				return GraphExpressionOutput->Id == ExpressionOutput->Id;
 			});
-		ExpressionNode = GraphExpression ? (*GraphExpression)->GraphNode : nullptr;
+		ExpressionNode = GraphExpression ? ToRawPtr((*GraphExpression)->GraphNode) : nullptr;
 	}
 
 	JumpToNode(ExpressionNode);
@@ -5871,7 +5871,7 @@ void FMaterialEditor::OnExpandNodes()
 		MoveNodesToAveragePos(ExpandedNodes, FVector2D(SourceNode->NodePosX, SourceNode->NodePosY), bExpandedNodesNeedUniqueGuid);
 	}
 
-	UMaterialExpression* SubgraphExpression = FocusedGraphEd ? Cast<UMaterialGraph>(FocusedGraphEd->GetCurrentGraph())->SubgraphExpression : Material->MaterialGraph->SubgraphExpression;
+	UMaterialExpression* SubgraphExpression = FocusedGraphEd ? ToRawPtr(Cast<UMaterialGraph>(FocusedGraphEd->GetCurrentGraph())->SubgraphExpression) : ToRawPtr(Material->MaterialGraph->SubgraphExpression);
 
 	for (UEdGraphNode* ExpandedNode : ExpandedNodes)
 	{
diff --git a/Engine/Source/Editor/MaterialEditor/Private/MaterialInstanceEditor.cpp b/Engine/Source/Editor/MaterialEditor/Private/MaterialInstanceEditor.cpp
--- a/Engine/Source/Editor/MaterialEditor/Private/MaterialInstanceEditor.cpp
+++ b/Engine/Source/Editor/MaterialEditor/Private/MaterialInstanceEditor.cpp
@@ -1110,7 +1110,7 @@ void FMaterialInstanceEditor::RefreshPreviewAsset()
 		UMaterialInterface* ParentMaterial = MaterialEditorInstance->SourceInstance->Parent;
 
 		UObject* ParentPreview = ParentMaterial != nullptr ? ParentMaterial->PreviewMesh.TryLoad() : nullptr;
-		PreviewAsset = ParentPreview != nullptr ? ParentPreview : GUnrealEd->GetThumbnailManager()->EditorSphere;
+		PreviewAsset = ParentPreview != nullptr ? ParentPreview : ToRawPtr(GUnrealEd->GetThumbnailManager()->EditorSphere);
 
 		USceneThumbnailInfoWithPrimitive* ThumbnailInfo = Cast<USceneThumbnailInfoWithPrimitive>(MaterialEditorInstance->SourceInstance->ThumbnailInfo);
 		if (ThumbnailInfo)
@@ -1148,7 +1148,7 @@ void FMaterialInstanceEditor::RebuildInheritanceList()
 			FunctionParentList.Insert(Parent, 0);
 
 			Current = Cast<UMaterialFunctionInstance>(Parent);
-			Parent = Current ? Current->Parent : nullptr;
+			Parent = Current ? ToRawPtr(Current->Parent) : nullptr;
 		}
 	}
 	else
diff --git a/Engine/Source/Editor/MovieSceneTools/Private/TrackEditors/CameraAnimTrackEditor.cpp b/Engine/Source/Editor/MovieSceneTools/Private/TrackEditors/CameraAnimTrackEditor.cpp
--- a/Engine/Source/Editor/MovieSceneTools/Private/TrackEditors/CameraAnimTrackEditor.cpp
+++ b/Engine/Source/Editor/MovieSceneTools/Private/TrackEditors/CameraAnimTrackEditor.cpp
@@ -37,7 +37,7 @@ public:
 	virtual FText GetSectionTitle() const override 
 	{ 
 		UMovieSceneCameraAnimSection const* const AnimSection = Cast<UMovieSceneCameraAnimSection>(WeakSection.Get());
-		UCameraAnim const* const Anim = AnimSection ? AnimSection->AnimData.CameraAnim : nullptr;
+		UCameraAnim const* const Anim = AnimSection ? ToRawPtr(AnimSection->AnimData.CameraAnim) : nullptr;
 		if (Anim)
 		{
 			return FText::FromString(Anim->GetName());
diff --git a/Engine/Source/Editor/Persona/Private/AnimTimeline/AnimModel_AnimComposite.cpp b/Engine/Source/Editor/Persona/Private/AnimTimeline/AnimModel_AnimComposite.cpp
--- a/Engine/Source/Editor/Persona/Private/AnimTimeline/AnimModel_AnimComposite.cpp
+++ b/Engine/Source/Editor/Persona/Private/AnimTimeline/AnimModel_AnimComposite.cpp
@@ -84,7 +84,7 @@ void FAnimModel_AnimComposite::RecalculateSequenceLength()
 			// Reset view if we changed length (note: has to be done after ->SetCompositeLength)!
 			UpdateRange();
 
-			UAnimPreviewInstance* PreviewInstance = (GetPreviewScene()->GetPreviewMeshComponent()) ? GetPreviewScene()->GetPreviewMeshComponent()->PreviewInstance : nullptr;
+			UAnimPreviewInstance* PreviewInstance = (GetPreviewScene()->GetPreviewMeshComponent()) ? ToRawPtr(GetPreviewScene()->GetPreviewMeshComponent()->PreviewInstance) : nullptr;
 			if (PreviewInstance)
 			{
 				// Re-set the position, so instance is clamped properly
diff --git a/Engine/Source/Editor/Persona/Private/AnimTimeline/AnimModel_AnimMontage.cpp b/Engine/Source/Editor/Persona/Private/AnimTimeline/AnimModel_AnimMontage.cpp
--- a/Engine/Source/Editor/Persona/Private/AnimTimeline/AnimModel_AnimMontage.cpp
+++ b/Engine/Source/Editor/Persona/Private/AnimTimeline/AnimModel_AnimMontage.cpp
@@ -237,7 +237,7 @@ void FAnimModel_AnimMontage::RecalculateSequenceLength()
 			// Reset view if we changed length (note: has to be done after ->SetCompositeLength)!
 			UpdateRange();
 
-			UAnimPreviewInstance* PreviewInstance = (GetPreviewScene()->GetPreviewMeshComponent()) ? GetPreviewScene()->GetPreviewMeshComponent()->PreviewInstance : nullptr;
+			UAnimPreviewInstance* PreviewInstance = (GetPreviewScene()->GetPreviewMeshComponent()) ? ToRawPtr(GetPreviewScene()->GetPreviewMeshComponent()->PreviewInstance) : nullptr;
 			if (PreviewInstance)
 			{
 				// Re-set the position, so instance is clamped properly
diff --git a/Engine/Source/Editor/Persona/Private/SAnimNotifyPanel.cpp b/Engine/Source/Editor/Persona/Private/SAnimNotifyPanel.cpp
--- a/Engine/Source/Editor/Persona/Private/SAnimNotifyPanel.cpp
+++ b/Engine/Source/Editor/Persona/Private/SAnimNotifyPanel.cpp
@@ -3111,7 +3111,7 @@ TSharedPtr<SWidget> SAnimNotifyTrack::SummonContextMenu(const FGeometry& MyGeome
 	if (NotifyEvent)
 	{
 		UObject* NotifyObject = NotifyEvent->Notify;
-		NotifyObject = NotifyObject ? NotifyObject : NotifyEvent->NotifyStateClass;
+		NotifyObject = NotifyObject ? NotifyObject : ToRawPtr(NotifyEvent->NotifyStateClass);
 
 		MenuBuilder.BeginSection("ViewSource", LOCTEXT("NotifyViewHeading", "View"));
 
diff --git a/Engine/Source/Editor/Persona/Private/SRigWindow.cpp b/Engine/Source/Editor/Persona/Private/SRigWindow.cpp
--- a/Engine/Source/Editor/Persona/Private/SRigWindow.cpp
+++ b/Engine/Source/Editor/Persona/Private/SRigWindow.cpp
@@ -304,7 +304,7 @@ const struct FReferenceSkeleton& SRigWindow::GetReferenceSkeleton() const
 {
 	// have to change this to preview mesh because that's what the retarget base pose will be
 	UDebugSkelMeshComponent* PreviewMeshComp = PreviewScenePtr.Pin()->GetPreviewMeshComponent();
-	USkeletalMesh* PreviewMesh = (PreviewMeshComp) ? PreviewMeshComp->SkeletalMesh : nullptr;
+	USkeletalMesh* PreviewMesh = (PreviewMeshComp) ? ToRawPtr(PreviewMeshComp->SkeletalMesh) : nullptr;
 	// it's because retarget base pose leaves in mesh, so if you give ref skeleton of skeleton, you might have joint that your mesh doesn't have
 	return (PreviewMesh)? PreviewMesh->GetRefSkeleton() : EditableSkeletonPtr.Pin()->GetSkeleton().GetReferenceSkeleton();
 }
diff --git a/Engine/Source/Editor/PhysicsAssetEditor/Private/PhysicsAssetEditorSharedData.cpp b/Engine/Source/Editor/PhysicsAssetEditor/Private/PhysicsAssetEditorSharedData.cpp
--- a/Engine/Source/Editor/PhysicsAssetEditor/Private/PhysicsAssetEditorSharedData.cpp
+++ b/Engine/Source/Editor/PhysicsAssetEditor/Private/PhysicsAssetEditorSharedData.cpp
@@ -1625,7 +1625,7 @@ void FPhysicsAssetEditorSharedData::DeleteBody(int32 DelBodyIndex, bool bRefresh
 	{
 		int32 ParentBodyIndex = PhysicsAsset->FindParentBodyIndex(EditorSkelMesh, BoneIndex);
 
-		UBodySetup * ParentBody = ParentBodyIndex != INDEX_NONE ? PhysicsAsset->SkeletalBodySetups[ParentBodyIndex] : NULL;
+		UBodySetup * ParentBody = ParentBodyIndex != INDEX_NONE ? ToRawPtr(PhysicsAsset->SkeletalBodySetups[ParentBodyIndex]) : NULL;
 
 		for (const int32 ConstraintIndex : Constraints)
 		{
diff --git a/Engine/Source/Editor/SequenceRecorder/Private/AnimationRecorder.cpp b/Engine/Source/Editor/SequenceRecorder/Private/AnimationRecorder.cpp
--- a/Engine/Source/Editor/SequenceRecorder/Private/AnimationRecorder.cpp
+++ b/Engine/Source/Editor/SequenceRecorder/Private/AnimationRecorder.cpp
@@ -656,7 +656,7 @@ bool FAnimationRecorder::Record(USkeletalMeshComponent* Component, FTransform co
 	if (ensure(AnimationObject))
 	{
 		UAnimDataController* Controller = AnimationObject->GetController();
-		USkeletalMesh* SkeletalMesh = Component->MasterPoseComponent != nullptr ? Component->MasterPoseComponent->SkeletalMesh : Component->SkeletalMesh;
+		USkeletalMesh* SkeletalMesh = Component->MasterPoseComponent != nullptr ? ToRawPtr(Component->MasterPoseComponent->SkeletalMesh) : ToRawPtr(Component->SkeletalMesh);
 
 		const TArray<FBoneAnimationTrack>& BoneAnimationTracks = AnimationObject->GetDataModel()->GetBoneAnimationTracks();
 
diff --git a/Engine/Source/Editor/SequenceRecorder/Private/Sections/MovieScene3DTransformSectionRecorder.cpp b/Engine/Source/Editor/SequenceRecorder/Private/Sections/MovieScene3DTransformSectionRecorder.cpp
--- a/Engine/Source/Editor/SequenceRecorder/Private/Sections/MovieScene3DTransformSectionRecorder.cpp
+++ b/Engine/Source/Editor/SequenceRecorder/Private/Sections/MovieScene3DTransformSectionRecorder.cpp
@@ -132,7 +132,7 @@ void FMovieScene3DTransformSectionRecorder::FinalizeSection(float CurrentTime)
 		USkeletalMeshComponent* SkeletalMeshComponent = AnimRecorder->GetSkeletalMeshComponent();
 		if (SkeletalMeshComponent)
 		{
-			USkeletalMesh* SkeletalMesh = SkeletalMeshComponent->MasterPoseComponent != nullptr ? SkeletalMeshComponent->MasterPoseComponent->SkeletalMesh : SkeletalMeshComponent->SkeletalMesh;
+			USkeletalMesh* SkeletalMesh = SkeletalMeshComponent->MasterPoseComponent != nullptr ? ToRawPtr(SkeletalMeshComponent->MasterPoseComponent->SkeletalMesh) : ToRawPtr(SkeletalMeshComponent->SkeletalMesh);
 			if (AnimSequence && SkeletalMesh)
 			{
 				// find the root bone
diff --git a/Engine/Source/Editor/SkeletonEditor/Private/SSkeletonTree.cpp b/Engine/Source/Editor/SkeletonEditor/Private/SSkeletonTree.cpp
--- a/Engine/Source/Editor/SkeletonEditor/Private/SSkeletonTree.cpp
+++ b/Engine/Source/Editor/SkeletonEditor/Private/SSkeletonTree.cpp
@@ -1241,7 +1241,7 @@ void SSkeletonTree::OnPasteSockets(bool bPasteToSelectedBone)
 	if ( TreeSelection.IsSingleOfTypeSelected<FSkeletonTreeBoneItem>() )
 	{
 		FName DestBoneName = bPasteToSelectedBone ? TreeSelection.GetSingleSelectedItem()->GetRowItemName() : NAME_None;
-		USkeletalMesh* SkeletalMesh = GetPreviewScene().IsValid() ? GetPreviewScene()->GetPreviewMeshComponent()->SkeletalMesh : nullptr;
+		USkeletalMesh* SkeletalMesh = GetPreviewScene().IsValid() ? ToRawPtr(GetPreviewScene()->GetPreviewMeshComponent()->SkeletalMesh) : nullptr;
 		GetEditableSkeletonInternal()->HandlePasteSockets(DestBoneName, SkeletalMesh);
 
 		CreateFromSkeleton();
@@ -1302,7 +1302,7 @@ void SSkeletonTree::OnCustomizeSocket()
 	if(TreeSelection.IsSingleOfTypeSelected<FSkeletonTreeSocketItem>())
 	{
 		USkeletalMeshSocket* SocketToCustomize = StaticCastSharedPtr<FSkeletonTreeSocketItem>(TreeSelection.GetSingleSelectedItem())->GetSocket();
-		USkeletalMesh* SkeletalMesh = GetPreviewScene().IsValid() ? GetPreviewScene()->GetPreviewMeshComponent()->SkeletalMesh : nullptr;
+		USkeletalMesh* SkeletalMesh = GetPreviewScene().IsValid() ? ToRawPtr(GetPreviewScene()->GetPreviewMeshComponent()->SkeletalMesh) : nullptr;
 		GetEditableSkeletonInternal()->HandleCustomizeSocket(SocketToCustomize, SkeletalMesh);
 		CreateFromSkeleton();
 	}
@@ -1372,7 +1372,7 @@ void  SSkeletonTree::OnRemoveAllAssets()
 
 bool SSkeletonTree::CanRemoveAllAssets() const
 {
-	USkeletalMesh* SkeletalMesh = GetPreviewScene().IsValid() ? GetPreviewScene()->GetPreviewMeshComponent()->SkeletalMesh : nullptr;
+	USkeletalMesh* SkeletalMesh = GetPreviewScene().IsValid() ? ToRawPtr(GetPreviewScene()->GetPreviewMeshComponent()->SkeletalMesh) : nullptr;
 
 	const bool bHasPreviewAttachedObjects = GetEditableSkeletonInternal()->GetSkeleton().PreviewAttachedAssetContainer.Num() > 0;
 	const bool bHasMeshPreviewAttachedObjects = ( SkeletalMesh && SkeletalMesh->GetPreviewAttachedAssetContainer().Num() );
@@ -2235,7 +2235,7 @@ void SSkeletonTree::SelectItemsBy(TFunctionRef<bool(const TSharedRef<ISkeletonTr
 
 void SSkeletonTree::DuplicateAndSelectSocket(const FSelectedSocketInfo& SocketInfoToDuplicate, const FName& NewParentBoneName /*= FName()*/)
 {
-	USkeletalMesh* SkeletalMesh = GetPreviewScene().IsValid() ? GetPreviewScene()->GetPreviewMeshComponent()->SkeletalMesh : nullptr;
+	USkeletalMesh* SkeletalMesh = GetPreviewScene().IsValid() ? ToRawPtr(GetPreviewScene()->GetPreviewMeshComponent()->SkeletalMesh) : nullptr;
 	USkeletalMeshSocket* NewSocket = GetEditableSkeleton()->DuplicateSocket(SocketInfoToDuplicate, NewParentBoneName, SkeletalMesh);
 
 	if (GetPreviewScene().IsValid())
diff --git a/Engine/Source/Editor/SkeletonEditor/Private/SkeletonTreeBoneItem.cpp b/Engine/Source/Editor/SkeletonEditor/Private/SkeletonTreeBoneItem.cpp
--- a/Engine/Source/Editor/SkeletonEditor/Private/SkeletonTreeBoneItem.cpp
+++ b/Engine/Source/Editor/SkeletonEditor/Private/SkeletonTreeBoneItem.cpp
@@ -427,7 +427,7 @@ FReply FSkeletonTreeBoneItem::HandleDrop(const FDragDropEvent& DragDropEvent)
 		else if (BoneName != SocketInfo.Socket->BoneName)
 		{
 			// The socket can be dropped here if we're a bone and NOT the socket's existing parent
-			USkeletalMesh* SkeletalMesh = GetSkeletonTree()->GetPreviewScene().IsValid() ? GetSkeletonTree()->GetPreviewScene()->GetPreviewMeshComponent()->SkeletalMesh : nullptr;
+			USkeletalMesh* SkeletalMesh = GetSkeletonTree()->GetPreviewScene().IsValid() ? ToRawPtr(GetSkeletonTree()->GetPreviewScene()->GetPreviewMeshComponent()->SkeletalMesh) : nullptr;
 			GetEditableSkeleton()->SetSocketParent(SocketInfo.Socket->SocketName, BoneName, SkeletalMesh);
 
 			return FReply::Handled();
diff --git a/Engine/Source/Editor/SkeletonEditor/Private/SkeletonTreeSocketItem.cpp b/Engine/Source/Editor/SkeletonEditor/Private/SkeletonTreeSocketItem.cpp
--- a/Engine/Source/Editor/SkeletonEditor/Private/SkeletonTreeSocketItem.cpp
+++ b/Engine/Source/Editor/SkeletonEditor/Private/SkeletonTreeSocketItem.cpp
@@ -172,7 +172,7 @@ bool FSkeletonTreeSocketItem::OnVerifySocketNameChanged( const FText& InText, FT
 	}
 	else
 	{
-		USkeletalMesh* SkeletalMesh = GetSkeletonTree()->GetPreviewScene().IsValid() ? GetSkeletonTree()->GetPreviewScene()->GetPreviewMeshComponent()->SkeletalMesh : nullptr;
+		USkeletalMesh* SkeletalMesh = GetSkeletonTree()->GetPreviewScene().IsValid() ? ToRawPtr(GetSkeletonTree()->GetPreviewScene()->GetPreviewMeshComponent()->SkeletalMesh) : nullptr;
 		bVerifyName = !GetEditableSkeleton()->DoesSocketAlreadyExist( Socket, NewText, ParentType, SkeletalMesh );
 
 		// Needs to be checked on verify.
@@ -192,7 +192,7 @@ void FSkeletonTreeSocketItem::OnCommitSocketName( const FText& InText, ETextComm
 	FText NewText = FText::TrimPrecedingAndTrailing(InText);
 
 	// Notify skeleton tree of socket rename
-	USkeletalMesh* SkeletalMesh = GetSkeletonTree()->GetPreviewScene().IsValid() ? GetSkeletonTree()->GetPreviewScene()->GetPreviewMeshComponent()->SkeletalMesh : nullptr;
+	USkeletalMesh* SkeletalMesh = GetSkeletonTree()->GetPreviewScene().IsValid() ? ToRawPtr(GetSkeletonTree()->GetPreviewScene()->GetPreviewMeshComponent()->SkeletalMesh) : nullptr;
 	GetEditableSkeleton()->RenameSocket(Socket->SocketName, FName(*NewText.ToString()), SkeletalMesh);
 }
 
diff --git a/Engine/Source/Editor/UMGEditor/Private/WidgetBlueprintEditorUtils.cpp b/Engine/Source/Editor/UMGEditor/Private/WidgetBlueprintEditorUtils.cpp
--- a/Engine/Source/Editor/UMGEditor/Private/WidgetBlueprintEditorUtils.cpp
+++ b/Engine/Source/Editor/UMGEditor/Private/WidgetBlueprintEditorUtils.cpp
@@ -580,7 +580,7 @@ void FWidgetBlueprintEditorUtils::FindAllAncestorNamedSlotHostWidgetsForContent(
 	OutSlotHostWidgets.Empty();
 	UUserWidget* Preview = BlueprintEditor->GetPreview();
 	UWidgetBlueprint* WidgetBP = BlueprintEditor->GetWidgetBlueprintObj();
-	UWidgetTree* WidgetTree = (WidgetBP != nullptr) ? WidgetBP->WidgetTree : nullptr;
+	UWidgetTree* WidgetTree = (WidgetBP != nullptr) ? ToRawPtr(WidgetBP->WidgetTree) : nullptr;
 
 	if (Preview != nullptr && WidgetTree != nullptr)
 	{
diff --git a/Engine/Source/Editor/UnrealEd/Private/Fbx/ReimportFbxSceneFactory.cpp b/Engine/Source/Editor/UnrealEd/Private/Fbx/ReimportFbxSceneFactory.cpp
--- a/Engine/Source/Editor/UnrealEd/Private/Fbx/ReimportFbxSceneFactory.cpp
+++ b/Engine/Source/Editor/UnrealEd/Private/Fbx/ReimportFbxSceneFactory.cpp
@@ -105,7 +105,7 @@ UFbxSceneImportData *GetFbxSceneImportData(UObject *Obj)
 
 		if (ImportData != nullptr)
 		{
-			SceneImportData = ImportData->bImportAsScene ? ImportData->FbxSceneImportDataReference : nullptr;
+			SceneImportData = ImportData->bImportAsScene ? ToRawPtr(ImportData->FbxSceneImportDataReference) : nullptr;
 		}
 	}
 	return SceneImportData;
diff --git a/Engine/Source/Editor/UnrealEd/Private/Kismet2/BlueprintEditorUtils.cpp b/Engine/Source/Editor/UnrealEd/Private/Kismet2/BlueprintEditorUtils.cpp
--- a/Engine/Source/Editor/UnrealEd/Private/Kismet2/BlueprintEditorUtils.cpp
+++ b/Engine/Source/Editor/UnrealEd/Private/Kismet2/BlueprintEditorUtils.cpp
@@ -3170,7 +3170,7 @@ bool FBlueprintEditorUtils::IsDataOnlyBlueprint(const UBlueprint* Blueprint)
 	}
 
 	// Make sure there's nothing in the user construction script, other than an entry node
-	UEdGraph* UserConstructionScript = (Blueprint->FunctionGraphs.Num() == 1) ? Blueprint->FunctionGraphs[0] : nullptr;
+	UEdGraph* UserConstructionScript = (Blueprint->FunctionGraphs.Num() == 1) ? ToRawPtr(Blueprint->FunctionGraphs[0]) : nullptr;
 	if (UserConstructionScript && Blueprint->ParentClass)
 	{
 		//Call parent construction script may be added automatically
diff --git a/Engine/Source/Editor/UnrealEd/Private/Kismet2/KismetReinstanceUtilities.cpp b/Engine/Source/Editor/UnrealEd/Private/Kismet2/KismetReinstanceUtilities.cpp
--- a/Engine/Source/Editor/UnrealEd/Private/Kismet2/KismetReinstanceUtilities.cpp
+++ b/Engine/Source/Editor/UnrealEd/Private/Kismet2/KismetReinstanceUtilities.cpp
@@ -1676,7 +1676,7 @@ UClass* FBlueprintCompileReinstancer::MoveCDOToNewClass(UClass* OwnerClass, cons
 	}
 #endif
 
-	UFunction* DuplicatedClassUberGraphFunction = BPGDuplicatedClass ? BPGDuplicatedClass->UberGraphFunction : nullptr;
+	UFunction* DuplicatedClassUberGraphFunction = BPGDuplicatedClass ? ToRawPtr(BPGDuplicatedClass->UberGraphFunction) : nullptr;
 	if (DuplicatedClassUberGraphFunction)
 	{
 		DuplicatedClassUberGraphFunction->Bind();
diff --git a/Engine/Source/Editor/UnrealEd/Private/ObjectTools.cpp b/Engine/Source/Editor/UnrealEd/Private/ObjectTools.cpp
--- a/Engine/Source/Editor/UnrealEd/Private/ObjectTools.cpp
+++ b/Engine/Source/Editor/UnrealEd/Private/ObjectTools.cpp
@@ -196,7 +196,7 @@ namespace ObjectTools
 
 		if (!CVarUseLegacyGetReferencersForDeletion.GetValueOnAnyThread())
 		{
-			const UTransactor* Transactor = GEditor ? GEditor->Trans : nullptr;
+			const UTransactor* Transactor = GEditor ? ToRawPtr(GEditor->Trans) : nullptr;
 			bool bIsGatheringPackageRef = InObject->IsA<UPackage>();
 
 			// Get the cluster of objects that are going to be deleted
diff --git a/Engine/Source/Editor/VirtualTexturingEditor/Private/VirtualTextureBuilderThumbnailRenderer.cpp b/Engine/Source/Editor/VirtualTexturingEditor/Private/VirtualTextureBuilderThumbnailRenderer.cpp
--- a/Engine/Source/Editor/VirtualTexturingEditor/Private/VirtualTextureBuilderThumbnailRenderer.cpp
+++ b/Engine/Source/Editor/VirtualTexturingEditor/Private/VirtualTextureBuilderThumbnailRenderer.cpp
@@ -12,14 +12,14 @@ UVirtualTextureBuilderThumbnailRenderer::UVirtualTextureBuilderThumbnailRenderer
 bool UVirtualTextureBuilderThumbnailRenderer::CanVisualizeAsset(UObject* Object)
 {
 	UVirtualTextureBuilder* VirtualTextureBuilder = Cast<UVirtualTextureBuilder>(Object);
-	UVirtualTexture2D* Texture = VirtualTextureBuilder != nullptr ? VirtualTextureBuilder->Texture : nullptr;
+	UVirtualTexture2D* Texture = VirtualTextureBuilder != nullptr ? ToRawPtr(VirtualTextureBuilder->Texture) : nullptr;
 	return Texture != nullptr ? UTextureThumbnailRenderer::CanVisualizeAsset(Texture) : false;
 }
 
 void UVirtualTextureBuilderThumbnailRenderer::Draw(UObject* Object, int32 X, int32 Y, uint32 Width, uint32 Height, FRenderTarget* RenderTarget, FCanvas* Canvas, bool bAdditionalViewFamily)
 {
 	UVirtualTextureBuilder* VirtualTextureBuilder = Cast<UVirtualTextureBuilder>(Object);
-	UVirtualTexture2D* Texture = VirtualTextureBuilder != nullptr ? VirtualTextureBuilder->Texture : nullptr;
+	UVirtualTexture2D* Texture = VirtualTextureBuilder != nullptr ? ToRawPtr(VirtualTextureBuilder->Texture) : nullptr;
 	if (Texture != nullptr)
 	{
 		UTextureThumbnailRenderer::Draw(Texture, X, Y, Width, Height, RenderTarget, Canvas, bAdditionalViewFamily);
diff --git a/Engine/Source/Runtime/AIModule/Private/BehaviorTree/BTNode.cpp b/Engine/Source/Runtime/AIModule/Private/BehaviorTree/BTNode.cpp
--- a/Engine/Source/Runtime/AIModule/Private/BehaviorTree/BTNode.cpp
+++ b/Engine/Source/Runtime/AIModule/Private/BehaviorTree/BTNode.cpp
@@ -87,7 +87,7 @@ void UBTNode::InitializeInSubtree(UBehaviorTreeComponent& OwnerComp, uint8* Node
 		// composite nodes can't be instanced!
 		check(IsA(UBTCompositeNode::StaticClass()) == false);
 
-		UBTNode* NodeInstance = OwnerComp.NodeInstances.IsValidIndex(NextInstancedIndex) ? OwnerComp.NodeInstances[NextInstancedIndex] : NULL;
+		UBTNode* NodeInstance = OwnerComp.NodeInstances.IsValidIndex(NextInstancedIndex) ? ToRawPtr(OwnerComp.NodeInstances[NextInstancedIndex]) : NULL;
 		if (NodeInstance == NULL)
 		{
 			NodeInstance = (UBTNode*)StaticDuplicateObject(this, &OwnerComp);
diff --git a/Engine/Source/Runtime/AIModule/Private/BehaviorTree/BlackboardComponent.cpp b/Engine/Source/Runtime/AIModule/Private/BehaviorTree/BlackboardComponent.cpp
--- a/Engine/Source/Runtime/AIModule/Private/BehaviorTree/BlackboardComponent.cpp
+++ b/Engine/Source/Runtime/AIModule/Private/BehaviorTree/BlackboardComponent.cpp
@@ -246,8 +246,8 @@ void UBlackboardComponent::PopulateSynchronizedKeys()
 						uint8* RawData = GetKeyRawData(KeyID) + DataOffset;
 						uint8* RawSource = OtherBlackboard->GetKeyRawData(OtherKeyID) + DataOffset;
 
-						UBlackboardKeyType* KeyOb = bKeyHasInstance ? KeyInstances[KeyID] : Key.KeyType;
-						const UBlackboardKeyType* SourceKeyOb = bKeyHasInstance ? OtherBlackboard->KeyInstances[OtherKeyID] : Key.KeyType;
+						UBlackboardKeyType* KeyOb = bKeyHasInstance ? ToRawPtr(KeyInstances[KeyID]) : ToRawPtr(Key.KeyType);
+						const UBlackboardKeyType* SourceKeyOb = bKeyHasInstance ? ToRawPtr(OtherBlackboard->KeyInstances[OtherKeyID]) : ToRawPtr(Key.KeyType);
 
 						KeyOb->CopyValues(*this, RawData, SourceKeyOb, RawSource);
 					}
diff --git a/Engine/Source/Runtime/Engine/Private/Blueprint.cpp b/Engine/Source/Runtime/Engine/Private/Blueprint.cpp
--- a/Engine/Source/Runtime/Engine/Private/Blueprint.cpp
+++ b/Engine/Source/Runtime/Engine/Private/Blueprint.cpp
@@ -1717,7 +1717,7 @@ bool UBlueprint::ChangeOwnerOfTemplates()
 
 			for (USCS_Node* SCSNode : SCS->GetAllNodes())
 			{
-				UActorComponent* Component = SCSNode ? SCSNode->ComponentTemplate : NULL;
+				UActorComponent* Component = SCSNode ? ToRawPtr(SCSNode->ComponentTemplate) : NULL;
 				if (Component && Component->GetOuter() == this)
 				{
 					const bool bRenamed = Component->Rename(FUniqueNewNameHelper(Component->GetName(), BPGClass).Get(), BPGClass, REN_ForceNoResetLoaders | REN_DoNotDirty);
diff --git a/Engine/Source/Runtime/Engine/Private/Brush.cpp b/Engine/Source/Runtime/Engine/Private/Brush.cpp
--- a/Engine/Source/Runtime/Engine/Private/Brush.cpp
+++ b/Engine/Source/Runtime/Engine/Private/Brush.cpp
@@ -158,7 +158,7 @@ void ABrush::SetIsTemporarilyHiddenInEditor( bool bIsHidden )
 		Super::SetIsTemporarilyHiddenInEditor(bIsHidden);
 		
 		ULevel* Level = GetLevel();
-		UModel* Model = Level ? Level->Model : nullptr;
+		UModel* Model = Level ? ToRawPtr(Level->Model) : nullptr;
 
 		if (Level && Model)
 		{
diff --git a/Engine/Source/Runtime/Engine/Private/ComponentInstanceDataCache.cpp b/Engine/Source/Runtime/Engine/Private/ComponentInstanceDataCache.cpp
--- a/Engine/Source/Runtime/Engine/Private/ComponentInstanceDataCache.cpp
+++ b/Engine/Source/Runtime/Engine/Private/ComponentInstanceDataCache.cpp
@@ -192,7 +192,7 @@ public:
 		// UObject pointer are serialized as Index in ActorInstanceData
 		int32 ReferenceIndex = INDEX_NONE;
 		*this << ReferenceIndex;
-		Object = ActorInstanceData.ReferencedObjects.IsValidIndex(ReferenceIndex) ? ActorInstanceData.ReferencedObjects[ReferenceIndex] : nullptr;
+		Object = ActorInstanceData.ReferencedObjects.IsValidIndex(ReferenceIndex) ? ToRawPtr(ActorInstanceData.ReferencedObjects[ReferenceIndex]) : nullptr;
 		return *this;
 	}
 
diff --git a/Engine/Source/Runtime/Engine/Private/Components/CharacterMovementComponent.cpp b/Engine/Source/Runtime/Engine/Private/Components/CharacterMovementComponent.cpp
--- a/Engine/Source/Runtime/Engine/Private/Components/CharacterMovementComponent.cpp
+++ b/Engine/Source/Runtime/Engine/Private/Components/CharacterMovementComponent.cpp
@@ -8263,7 +8263,7 @@ void UCharacterMovementComponent::CallServerMovePacked(const FSavedMove_Characte
 	// Extract the net package map used for serializing object references.
 	APlayerController* PC = Cast<APlayerController>(CharacterOwner->GetController());
 	UNetConnection* NetConnection = PC ? PC->GetNetConnection() : nullptr;
-	ServerMoveBitWriter.PackageMap = NetConnection ? NetConnection->PackageMap : nullptr;
+	ServerMoveBitWriter.PackageMap = NetConnection ? ToRawPtr(NetConnection->PackageMap) : nullptr;
 	if (ServerMoveBitWriter.PackageMap == nullptr)
 	{
 		UE_LOG(LogNetPlayerMovement, Error, TEXT("CallServerMovePacked: Failed to find a NetConnection/PackageMap for data serialization!"));
@@ -9584,7 +9584,7 @@ void UCharacterMovementComponent::ServerSendMoveResponse(const FClientAdjustment
 	// Extract the net package map used for serializing object references.
 	APlayerController* PC = Cast<APlayerController>(CharacterOwner->GetController());
 	UNetConnection* NetConnection = PC ? PC->GetNetConnection() : nullptr;
-	MoveResponseBitWriter.PackageMap = NetConnection ? NetConnection->PackageMap : nullptr;
+	MoveResponseBitWriter.PackageMap = NetConnection ? ToRawPtr(NetConnection->PackageMap) : nullptr;
 	if (MoveResponseBitWriter.PackageMap == nullptr)
 	{
 		UE_LOG(LogNetPlayerMovement, Error, TEXT("ServerSendMoveResponse: Failed to find a NetConnection/PackageMap for data serialization!"));
@@ -11535,7 +11535,7 @@ bool UCharacterMovementComponent::CanDelaySendingMove(const FSavedMovePtr& NewMo
 
 float UCharacterMovementComponent::GetClientNetSendDeltaTime(const APlayerController* PC, const FNetworkPredictionData_Client_Character* ClientData, const FSavedMovePtr& NewMove) const
 {
-	const UPlayer* Player = (PC ? PC->Player : nullptr);
+	const UPlayer* Player = (PC ? ToRawPtr(PC->Player) : nullptr);
 	const UWorld* MyWorld = GetWorld();
 	const AGameStateBase* const GameState = MyWorld->GetGameState();
 	const AGameNetworkManager* GameNetworkManager = (const AGameNetworkManager*)(AGameNetworkManager::StaticClass()->GetDefaultObject());
diff --git a/Engine/Source/Runtime/Engine/Private/Components/SkeletalMeshComponent.cpp b/Engine/Source/Runtime/Engine/Private/Components/SkeletalMeshComponent.cpp
--- a/Engine/Source/Runtime/Engine/Private/Components/SkeletalMeshComponent.cpp
+++ b/Engine/Source/Runtime/Engine/Private/Components/SkeletalMeshComponent.cpp
@@ -2246,7 +2246,7 @@ void USkeletalMeshComponent::RefreshBoneTransforms(FActorComponentTickFunction*
 
 	AnimEvaluationContext.SkeletalMesh = SkeletalMesh;
 	AnimEvaluationContext.AnimInstance = AnimScriptInstance;
-	AnimEvaluationContext.PostProcessAnimInstance = (ShouldEvaluatePostProcessInstance())? PostProcessAnimInstance: nullptr;
+	AnimEvaluationContext.PostProcessAnimInstance = (ShouldEvaluatePostProcessInstance())? ToRawPtr(PostProcessAnimInstance): nullptr;
 
 	if (CurrentAnimCurveUIDFinder)
 	{
diff --git a/Engine/Source/Runtime/Engine/Private/DataChannel.cpp b/Engine/Source/Runtime/Engine/Private/DataChannel.cpp
--- a/Engine/Source/Runtime/Engine/Private/DataChannel.cpp
+++ b/Engine/Source/Runtime/Engine/Private/DataChannel.cpp
@@ -2361,7 +2361,7 @@ void UActorChannel::NotifyActorChannelOpen(AActor* InActor, FInBunch& InBunch)
 				DemoNetDriver->GetNetworkObjectList().FindOrAdd(Actor, DemoNetDriver);
 				DemoNetDriver->FlushActorDormancy(Actor);
 
-				UNetConnection* DemoClientConnection = (DemoNetDriver->ClientConnections.Num() > 0) ? DemoNetDriver->ClientConnections[0] : nullptr;
+				UNetConnection* DemoClientConnection = (DemoNetDriver->ClientConnections.Num() > 0) ? ToRawPtr(DemoNetDriver->ClientConnections[0]) : nullptr;
 				if (DemoClientConnection)
 				{
 					DemoNetDriver->GetNetworkObjectList().MarkActive(Actor, DemoClientConnection, DemoNetDriver);
diff --git a/Engine/Source/Runtime/Engine/Private/DemoNetDriver.cpp b/Engine/Source/Runtime/Engine/Private/DemoNetDriver.cpp
--- a/Engine/Source/Runtime/Engine/Private/DemoNetDriver.cpp
+++ b/Engine/Source/Runtime/Engine/Private/DemoNetDriver.cpp
@@ -2921,7 +2921,7 @@ int32 UDemoNetDriver::CleanUpSplitscreenConnections(bool bDeleteOwner)
 
 	for (APlayerController* CurController : SpectatorControllers)
 	{
-		UNetConnection* ControllerNetConnection = (CurController != nullptr) ? CurController->NetConnection : nullptr;
+		UNetConnection* ControllerNetConnection = (CurController != nullptr) ? ToRawPtr(CurController->NetConnection) : nullptr;
 		if (ControllerNetConnection != nullptr && ControllerNetConnection->IsA(UChildConnection::StaticClass()))
 		{
 			++NumSplitscreenConnectionsCleaned;
diff --git a/Engine/Source/Runtime/Engine/Private/GameInstance.cpp b/Engine/Source/Runtime/Engine/Private/GameInstance.cpp
--- a/Engine/Source/Runtime/Engine/Private/GameInstance.cpp
+++ b/Engine/Source/Runtime/Engine/Private/GameInstance.cpp
@@ -1108,7 +1108,7 @@ bool UGameInstance::ClientTravelToSession(int32 ControllerId, FName InSessionNam
 	if (UOnlineEngineInterface::Get()->GetResolvedConnectString(World, InSessionName, URL))
 	{
 		ULocalPlayer* LP = GEngine->GetLocalPlayerFromControllerId(World, ControllerId);
-		APlayerController* PC = LP ? LP->PlayerController : nullptr;
+		APlayerController* PC = LP ? ToRawPtr(LP->PlayerController) : nullptr;
 		if (PC)
 		{
 			PC->ClientTravel(URL, TRAVEL_Absolute);
diff --git a/Engine/Source/Runtime/Engine/Private/HLODProxy.cpp b/Engine/Source/Runtime/Engine/Private/HLODProxy.cpp
--- a/Engine/Source/Runtime/Engine/Private/HLODProxy.cpp
+++ b/Engine/Source/Runtime/Engine/Private/HLODProxy.cpp
@@ -224,13 +224,13 @@ const AActor* UHLODProxy::FindFirstActor(const ALODActor* LODActor)
 {
 	auto RecursiveFindFirstActor = [&](const ALODActor* InLODActor)
 	{
-		const AActor* FirstActor = InLODActor->SubActors.IsValidIndex(0) ? InLODActor->SubActors[0] : nullptr;
+		const AActor* FirstActor = InLODActor->SubActors.IsValidIndex(0) ? ToRawPtr(InLODActor->SubActors[0]) : nullptr;
 		while (FirstActor != nullptr && FirstActor->IsA<ALODActor>())
 		{
 			const ALODActor* SubLODActor = Cast<ALODActor>(FirstActor);
 			if (SubLODActor)
 			{
-				FirstActor = SubLODActor->SubActors.IsValidIndex(0) ? SubLODActor->SubActors[0] : nullptr; 
+				FirstActor = SubLODActor->SubActors.IsValidIndex(0) ? ToRawPtr(SubLODActor->SubActors[0]) : nullptr; 
 			}
 			else
 			{
diff --git a/Engine/Source/Runtime/Engine/Private/Materials/MaterialInstance.cpp b/Engine/Source/Runtime/Engine/Private/Materials/MaterialInstance.cpp
--- a/Engine/Source/Runtime/Engine/Private/Materials/MaterialInstance.cpp
+++ b/Engine/Source/Runtime/Engine/Private/Materials/MaterialInstance.cpp
@@ -1347,7 +1347,7 @@ void UMaterialInstance::LogMaterialsAndTextures(FOutputDevice& Ar, int32 Indent)
 					MaterialInstanceToUse = CurrentMaterialInstance;
 				}
 
-				CurrentMaterialInterface = CurrentMaterialInstance ? CurrentMaterialInstance->Parent : nullptr;
+				CurrentMaterialInterface = CurrentMaterialInstance ? ToRawPtr(CurrentMaterialInstance->Parent) : nullptr;
 			}
 		}
 
@@ -4517,7 +4517,7 @@ bool UMaterialInstance::IsChildOf(const UMaterialInterface* ParentMaterialInterf
 	while (Material != ParentMaterialInterface && Material != nullptr)
 	{
 		const UMaterialInstance* MaterialInstance = Cast<const UMaterialInstance>(Material);
-		Material = (MaterialInstance != nullptr) ? MaterialInstance->Parent : nullptr;
+		Material = (MaterialInstance != nullptr) ? ToRawPtr(MaterialInstance->Parent) : nullptr;
 	}
 
 	return (Material != nullptr);
diff --git a/Engine/Source/Runtime/Engine/Private/NetDriver.cpp b/Engine/Source/Runtime/Engine/Private/NetDriver.cpp
--- a/Engine/Source/Runtime/Engine/Private/NetDriver.cpp
+++ b/Engine/Source/Runtime/Engine/Private/NetDriver.cpp
@@ -5801,7 +5801,7 @@ UReplicationDriver* UReplicationDriver::CreateReplicationDriver(UNetDriver* NetD
 		return nullptr;
 	}
 
-	UClass* ReplicationDriverClass = NetDriver ? NetDriver->ReplicationDriverClass : nullptr;
+	UClass* ReplicationDriverClass = NetDriver ? ToRawPtr(NetDriver->ReplicationDriverClass) : nullptr;
 	if (ReplicationDriverClass == nullptr)
 	{
 		UE_LOG(LogNet, Log, TEXT("ReplicationDriverClass is null! Not using ReplicationDriver."));
diff --git a/Engine/Source/Runtime/Engine/Private/PacketHandlers/StatelessConnectHandlerComponent.cpp b/Engine/Source/Runtime/Engine/Private/PacketHandlers/StatelessConnectHandlerComponent.cpp
--- a/Engine/Source/Runtime/Engine/Private/PacketHandlers/StatelessConnectHandlerComponent.cpp
+++ b/Engine/Source/Runtime/Engine/Private/PacketHandlers/StatelessConnectHandlerComponent.cpp
@@ -324,7 +324,7 @@ void StatelessConnectHandlerComponent::NotifyHandshakeBegin()
 {
 	if (Handler->Mode == Handler::Mode::Client)
 	{
-		UNetConnection* ServerConn = (Driver != nullptr ? Driver->ServerConnection : nullptr);
+		UNetConnection* ServerConn = (Driver != nullptr ? ToRawPtr(Driver->ServerConnection) : nullptr);
 
 		if (ServerConn != nullptr)
 		{
@@ -465,7 +465,7 @@ void StatelessConnectHandlerComponent::SendConnectChallenge(TSharedPtr<const FIn
 
 void StatelessConnectHandlerComponent::SendChallengeResponse(uint8 InSecretId, double InTimestamp, uint8 InCookie[COOKIE_BYTE_SIZE])
 {
-	UNetConnection* ServerConn = (Driver != nullptr ? Driver->ServerConnection : nullptr);
+	UNetConnection* ServerConn = (Driver != nullptr ? ToRawPtr(Driver->ServerConnection) : nullptr);
 
 	if (ServerConn != nullptr)
 	{
@@ -792,7 +792,7 @@ void StatelessConnectHandlerComponent::Incoming(FBitReader& Packet)
 					{
 						if (!bRestartedHandshake)
 						{
-							UNetConnection* ServerConn = (Driver != nullptr ? Driver->ServerConnection : nullptr);
+							UNetConnection* ServerConn = (Driver != nullptr ? ToRawPtr(Driver->ServerConnection) : nullptr);
 
 							// Extract the initial packet sequence from the random Cookie data
 							if (ensure(ServerConn != nullptr))
@@ -831,7 +831,7 @@ void StatelessConnectHandlerComponent::Incoming(FBitReader& Packet)
 
 						if (!bRestartedHandshake)
 						{
-							UNetConnection* ServerConn = (Driver != nullptr ? Driver->ServerConnection : nullptr);
+							UNetConnection* ServerConn = (Driver != nullptr ? ToRawPtr(Driver->ServerConnection) : nullptr);
 							double LastNetConnPacketTime = (ServerConn != nullptr ? ServerConn->LastReceiveRealtime : 0.0);
 
 							// The server may send multiple restart handshake packets, so have a 10 second delay between accepting them
diff --git a/Engine/Source/Runtime/Engine/Private/Particles/ParticleComponents.cpp b/Engine/Source/Runtime/Engine/Private/Particles/ParticleComponents.cpp
--- a/Engine/Source/Runtime/Engine/Private/Particles/ParticleComponents.cpp
+++ b/Engine/Source/Runtime/Engine/Private/Particles/ParticleComponents.cpp
@@ -5424,7 +5424,7 @@ void UParticleSystemComponent::FinalizeTickComponent()
 		}
 
 		UWorld* World = GetWorld();
-		AParticleEventManager* EventManager = (World ? World->MyParticleEventManager : NULL);
+		AParticleEventManager* EventManager = (World ? ToRawPtr(World->MyParticleEventManager) : NULL);
 		if (EventManager)
 		{
 			if (SpawnEvents.Num() > 0) EventManager->HandleParticleSpawnEvents(this, SpawnEvents);
diff --git a/Engine/Source/Runtime/Engine/Private/Player.cpp b/Engine/Source/Runtime/Engine/Private/Player.cpp
--- a/Engine/Source/Runtime/Engine/Private/Player.cpp
+++ b/Engine/Source/Runtime/Engine/Private/Player.cpp
@@ -95,7 +95,7 @@ bool UPlayer::Exec( UWorld* InWorld, const TCHAR* Cmd,FOutputDevice& Ar)
 	if (!ExecActor)
 	{
 		UNetConnection* NetConn = Cast<UNetConnection>(this);
-		ExecActor = (NetConn && NetConn->OwningActor) ? NetConn->OwningActor : nullptr;
+		ExecActor = (NetConn && NetConn->OwningActor) ? ToRawPtr(NetConn->OwningActor) : nullptr;
 	}
 
 	if (ExecActor)
diff --git a/Engine/Source/Runtime/Engine/Private/PlayerCameraManager.cpp b/Engine/Source/Runtime/Engine/Private/PlayerCameraManager.cpp
--- a/Engine/Source/Runtime/Engine/Private/PlayerCameraManager.cpp
+++ b/Engine/Source/Runtime/Engine/Private/PlayerCameraManager.cpp
@@ -363,7 +363,7 @@ UCameraAnimInst* APlayerCameraManager::AllocCameraAnimInst()
 {
 	check(IsInGameThread());
 
-	UCameraAnimInst* FreeAnim = (FreeAnims.Num() > 0) ? FreeAnims.Pop() : NULL;
+	UCameraAnimInst* FreeAnim = (FreeAnims.Num() > 0) ? ToRawPtr(FreeAnims.Pop()) : NULL;
 	if (FreeAnim)
 	{
 		UCameraAnimInst const* const DefaultInst = GetDefault<UCameraAnimInst>();
diff --git a/Engine/Source/Runtime/Engine/Private/SoundNodeWavePlayer.cpp b/Engine/Source/Runtime/Engine/Private/SoundNodeWavePlayer.cpp
--- a/Engine/Source/Runtime/Engine/Private/SoundNodeWavePlayer.cpp
+++ b/Engine/Source/Runtime/Engine/Private/SoundNodeWavePlayer.cpp
@@ -23,7 +23,7 @@ void USoundNodeWavePlayer::Serialize(FArchive& Ar)
 		}
 		else if (Ar.IsSaving())
 		{
-			USoundWave* HardReference = (ShouldHardReferenceAsset(Ar.CookingTarget()) ? SoundWave : nullptr);
+			USoundWave* HardReference = (ShouldHardReferenceAsset(Ar.CookingTarget()) ? ToRawPtr(SoundWave) : nullptr);
 			Ar << HardReference;
 		}
 	}
diff --git a/Engine/Source/Runtime/Engine/Private/UnrealEngine.cpp b/Engine/Source/Runtime/Engine/Private/UnrealEngine.cpp
--- a/Engine/Source/Runtime/Engine/Private/UnrealEngine.cpp
+++ b/Engine/Source/Runtime/Engine/Private/UnrealEngine.cpp
@@ -6543,7 +6543,7 @@ bool UEngine::HandleDumpParticleCountsCommand( const TCHAR* Cmd, FOutputDevice&
 					FParticleEmitterInstance* EmitInst = PSysComp->EmitterInstances[EmitterIndex];
 					if (EmitInst)
 					{
-						UParticleLODLevel* LODLevel = EmitInst->SpriteTemplate ? EmitInst->SpriteTemplate->LODLevels[0] : NULL;
+						UParticleLODLevel* LODLevel = EmitInst->SpriteTemplate ? ToRawPtr(EmitInst->SpriteTemplate->LODLevels[0]) : NULL;
 						if (bTrackUsageOnly == false)
 						{
 							Ar.Logf( TEXT("\t\tEmitter %2d:\tActive = %4d\tMaxActive = %4d"), 
@@ -13991,7 +13991,7 @@ static void AsyncMapChangeLevelLoadCompletionCallback(const FName& PackageName,
 			}
 		}
 
-		ULevel* Level = World ? World->PersistentLevel : nullptr;	
+		ULevel* Level = World ? ToRawPtr(World->PersistentLevel) : nullptr;	
 
 		// Print out a warning and set the error if we couldn't find a level in this package.
 		if( !Level )
diff --git a/Engine/Source/Runtime/Engine/Private/UserInterface/PlayerInput.cpp b/Engine/Source/Runtime/Engine/Private/UserInterface/PlayerInput.cpp
--- a/Engine/Source/Runtime/Engine/Private/UserInterface/PlayerInput.cpp
+++ b/Engine/Source/Runtime/Engine/Private/UserInterface/PlayerInput.cpp
@@ -2033,7 +2033,7 @@ bool UPlayerInput::ExecInputCommands( UWorld* InWorld, const TCHAR* Cmd, FOutput
 		if(CurrentEvent == IE_Pressed || (CurrentEvent == IE_Released && FParse::Command(&Str,TEXT("OnRelease"))))
 		{
 			APlayerController*	Actor = Cast<APlayerController>(GetOuter());
-			UPlayer* Player = Actor ? Actor->Player : NULL;
+			UPlayer* Player = Actor ? ToRawPtr(Actor->Player) : NULL;
 			if(ProcessConsoleExec(Str,Ar,this))
 			{
 				bResult = true;
diff --git a/Engine/Source/Runtime/Landscape/Private/LandscapeCollision.cpp b/Engine/Source/Runtime/Landscape/Private/LandscapeCollision.cpp
--- a/Engine/Source/Runtime/Landscape/Private/LandscapeCollision.cpp
+++ b/Engine/Source/Runtime/Landscape/Private/LandscapeCollision.cpp
@@ -980,7 +980,7 @@ bool ULandscapeHeightfieldCollisionComponent::CookCollisionData(const FName& For
 				if(!bUseDefMaterial)
 				{
 					uint8 DominantLayerIdx = DominantLayers ? DominantLayers[SrcSampleIndex] : -1;
-					ULandscapeLayerInfoObject* Layer = ComponentLayerInfos.IsValidIndex(DominantLayerIdx) ? ComponentLayerInfos[DominantLayerIdx] : nullptr;
+					ULandscapeLayerInfoObject* Layer = ComponentLayerInfos.IsValidIndex(DominantLayerIdx) ? ToRawPtr(ComponentLayerInfos[DominantLayerIdx]) : nullptr;
 
 					if(Layer == ALandscapeProxy::VisibilityLayer)
 					{
diff --git a/Engine/Source/Runtime/Landscape/Private/LandscapeRender.cpp b/Engine/Source/Runtime/Landscape/Private/LandscapeRender.cpp
--- a/Engine/Source/Runtime/Landscape/Private/LandscapeRender.cpp
+++ b/Engine/Source/Runtime/Landscape/Private/LandscapeRender.cpp
@@ -4139,7 +4139,7 @@ void ULandscapeComponent::GetStreamingRenderAssetInfo(FStreamingTextureLevelCont
 
 	for (int32 MaterialIndex = 0; MaterialIndex < MaterialInstanceCount; ++MaterialIndex)
 	{
-		const UMaterialInterface* MaterialInterface = FeatureLevel >= ERHIFeatureLevel::SM5 ? GetMaterialInstance(MaterialIndex) : MobileMaterialInterfaces[MaterialIndex];
+		const UMaterialInterface* MaterialInterface = FeatureLevel >= ERHIFeatureLevel::SM5 ? GetMaterialInstance(MaterialIndex) : ToRawPtr(MobileMaterialInterfaces[MaterialIndex]);
 
 		// Normal usage...
 		// Enumerate the textures used by the material.
diff --git a/Engine/Source/Runtime/UMG/Private/Animation/UMGSequencePlayer.cpp b/Engine/Source/Runtime/UMG/Private/Animation/UMGSequencePlayer.cpp
--- a/Engine/Source/Runtime/UMG/Private/Animation/UMGSequencePlayer.cpp
+++ b/Engine/Source/Runtime/UMG/Private/Animation/UMGSequencePlayer.cpp
@@ -145,7 +145,7 @@ void UUMGSequencePlayer::Tick(float DeltaTime)
 			if (bIsRunningTickManager)
 			{
 				UUserWidget* Widget = UserWidget.Get();
-				UUMGSequenceTickManager* TickManager = Widget ? Widget->AnimationTickManager : nullptr;
+				UUMGSequenceTickManager* TickManager = Widget ? ToRawPtr(Widget->AnimationTickManager) : nullptr;
 
 				if (!TickManager)
 				{
@@ -224,7 +224,7 @@ void UUMGSequencePlayer::PlayInternal(double StartAtTime, double EndAtTime, int3
 	PlayerStatus = EMovieScenePlayerStatus::Playing;
 
 	UUserWidget* Widget = UserWidget.Get();
-	UUMGSequenceTickManager* TickManager = Widget ? Widget->AnimationTickManager : nullptr;
+	UUMGSequenceTickManager* TickManager = Widget ? ToRawPtr(Widget->AnimationTickManager) : nullptr;
 
 	// Playback assumes the start frame has already been evaulated, so we also want to evaluate any events on the start frame here.
 	if (TickManager && RootTemplateInstance.IsValid())
@@ -279,7 +279,7 @@ void UUMGSequencePlayer::Pause()
 	if (RootTemplateInstance.HasEverUpdated())
 	{
 		UUserWidget* Widget = UserWidget.Get();
-		UUMGSequenceTickManager* TickManager = Widget ? Widget->AnimationTickManager : nullptr;
+		UUMGSequenceTickManager* TickManager = Widget ? ToRawPtr(Widget->AnimationTickManager) : nullptr;
 
 		if (TickManager)
 		{
@@ -308,7 +308,7 @@ void UUMGSequencePlayer::Stop()
 	PlayerStatus = EMovieScenePlayerStatus::Stopped;
 
 	UUserWidget* Widget = UserWidget.Get();
-	UUMGSequenceTickManager* TickManager = Widget ? Widget->AnimationTickManager : nullptr;
+	UUMGSequenceTickManager* TickManager = Widget ? ToRawPtr(Widget->AnimationTickManager) : nullptr;
 
 	if (TickManager && RootTemplateInstance.IsValid())
 	{
@@ -413,7 +413,7 @@ void UUMGSequencePlayer::QueueLatentAction(FMovieSceneSequenceLatentActionDelega
 	if (CVarUserWidgetUseParallelAnimation.GetValueOnGameThread())
 	{
 		UUserWidget* Widget = UserWidget.Get();
-		UUMGSequenceTickManager* TickManager = Widget ? Widget->AnimationTickManager : nullptr;
+		UUMGSequenceTickManager* TickManager = Widget ? ToRawPtr(Widget->AnimationTickManager) : nullptr;
 
 		if (ensure(TickManager))
 		{
@@ -431,7 +431,7 @@ void UUMGSequencePlayer::ApplyLatentActions()
 	if (CVarUserWidgetUseParallelAnimation.GetValueOnGameThread())
 	{
 		UUserWidget* Widget = UserWidget.Get();
-		UUMGSequenceTickManager* TickManager = Widget ? Widget->AnimationTickManager : nullptr;
+		UUMGSequenceTickManager* TickManager = Widget ? ToRawPtr(Widget->AnimationTickManager) : nullptr;
 		if (TickManager)
 		{
 			TickManager->RunLatentActions();
diff --git a/Engine/Source/Runtime/UMG/Private/Components/WidgetComponent.cpp b/Engine/Source/Runtime/UMG/Private/Components/WidgetComponent.cpp
--- a/Engine/Source/Runtime/UMG/Private/Components/WidgetComponent.cpp
+++ b/Engine/Source/Runtime/UMG/Private/Components/WidgetComponent.cpp
@@ -1162,7 +1162,7 @@ void UWidgetComponent::TickComponent(float DeltaTime, enum ELevelTick TickType,
 				UWorld* ThisWorld = GetWorld();
 
 				ULocalPlayer* TargetPlayer = GetOwnerPlayer();
-				APlayerController* PlayerController = TargetPlayer ? TargetPlayer->PlayerController : nullptr;
+				APlayerController* PlayerController = TargetPlayer ? ToRawPtr(TargetPlayer->PlayerController) : nullptr;
 
 				if ( TargetPlayer && PlayerController && IsVisible() && !(GetOwner()->IsHidden()))
 				{
