VHM: Move VT page table look up from compute shader to vertex shader.
Fixes some popping and seams in demo data.
Previous code had been trying to reduce VS work, but can't assume only 2 (or N where N is small) mips are required.

diff --git a/Engine/Plugins/Experimental/VirtualHeightfieldMesh/Shaders/Private/VirtualHeightfieldMesh.usf b/Engine/Plugins/Experimental/VirtualHeightfieldMesh/Shaders/Private/VirtualHeightfieldMesh.usf
--- a/Engine/Plugins/Experimental/VirtualHeightfieldMesh/Shaders/Private/VirtualHeightfieldMesh.usf
+++ b/Engine/Plugins/Experimental/VirtualHeightfieldMesh/Shaders/Private/VirtualHeightfieldMesh.usf
@@ -498,22 +498,11 @@ void CullInstancesCS( uint3 DispatchThreadId : SV_DispatchThreadID )
 		QuadRenderInstance OutInstance;
 		OutInstance.PosLevelPacked = Pos.x | (Pos.y << 12) | (Level << 24);
 
-		// Unpack physical address into full LocalToPhysicalUV here ready for use by vertex shader. 
-		// Could move unpacking to vertex shader if we want to balance memory bandwidth.
+		// Unpack physical address for first sample into full LocalToPhysicalUV here ready for use by vertex shader. 
 		OutInstance.UVTransform.xyz = GetVirtualToPhysicalUVTransform(Pos, Level, Item.PhysicalAddress, PhysicalPageTransform, NumPhysicalAddressBits);
 
-		[unroll]
-		for (int NeighborIndex = 0; NeighborIndex < 4; ++ NeighborIndex)
-		{
-			const uint PackedNeighborItem = QuadNeighborBuffer[QuadIndex * 4 + NeighborIndex];
-			const QuadNeighborItem NeighborItem = UnpackQuadNeighborItem(PackedNeighborItem);
-
-			OutInstance.NeigborUVTransform[NeighborIndex] = float4(GetVirtualToPhysicalUVTransform(Pos, Level, NeighborItem.PhysicalAddress, PhysicalPageTransform, NumPhysicalAddressBits), 0);
-		}
-
 		uint Write;
 		InterlockedAdd(RWIndirectArgsBuffer[1], 1, Write);
-
 		RWInstanceBuffer[Write] = OutInstance;
 	}
 }
diff --git a/Engine/Plugins/Experimental/VirtualHeightfieldMesh/Shaders/Private/VirtualHeightfieldMesh.ush b/Engine/Plugins/Experimental/VirtualHeightfieldMesh/Shaders/Private/VirtualHeightfieldMesh.ush
--- a/Engine/Plugins/Experimental/VirtualHeightfieldMesh/Shaders/Private/VirtualHeightfieldMesh.ush
+++ b/Engine/Plugins/Experimental/VirtualHeightfieldMesh/Shaders/Private/VirtualHeightfieldMesh.ush
@@ -121,10 +121,8 @@ QuadNeighborItem UnpackQuadNeighborItem(uint PackedItem)
 /** Final render instance description used by the DrawInstancedIndirect(). */
 struct QuadRenderInstance
 {
-	// Total of 64 bytes is nice for cache alignment. Take care when modifying.
-	uint PosLevelPacked;
 	float3 UVTransform;
-	float4 NeigborUVTransform[4];
+	uint PosLevelPacked;
 };
 
 uint2 UnpackPos(QuadRenderInstance Item)
diff --git a/Engine/Plugins/Experimental/VirtualHeightfieldMesh/Shaders/Private/VirtualHeightfieldMeshVertexFactory.ush b/Engine/Plugins/Experimental/VirtualHeightfieldMesh/Shaders/Private/VirtualHeightfieldMeshVertexFactory.ush
--- a/Engine/Plugins/Experimental/VirtualHeightfieldMesh/Shaders/Private/VirtualHeightfieldMeshVertexFactory.ush
+++ b/Engine/Plugins/Experimental/VirtualHeightfieldMesh/Shaders/Private/VirtualHeightfieldMeshVertexFactory.ush
@@ -1,13 +1,21 @@
 // Copyright Epic Games, Inc. All Rights Reserved.
 
 #include "/Engine/Private/VertexFactoryCommon.ush"
+#include "/Engine/Private/VirtualTextureCommon.ush"
 #include "VirtualHeightfieldMesh.ush"	
 
 StructuredBuffer<QuadRenderInstance> InstanceBuffer;
 
-Texture2D<float> HeightTexture;
+Texture2D<uint4> PageTableTexture;
+Texture2D<float4> HeightTexture;
 SamplerState HeightSampler;
+
+uint4 VTPackedUniform;
+uint4 VTPackedPageTableUniform0;
+uint4 VTPackedPageTableUniform1;
+
 float4 PageTableSize;
+float2 PhysicalTextureSize;
 float MaxLod;
 float4x4 VirtualHeightfieldToLocal;
 float4x4 VirtualHeightfieldToWorld;
@@ -83,18 +91,15 @@ FVertexFactoryIntermediates GetVertexFactoryIntermediates(FVertexFactoryInput In
 		(LocalUV.x < LocalUV.y ? 0.5 - LocalUV.x : 0.5 - LocalUV.y);
 	NeighborWeight *= 2;
 
-	float3 NeigborUVTransform = Item.NeigborUVTransform[NeighborIndex].xyz;
-	float2 NeighborPhysicalUV = NeigborUVTransform.xy + LocalUV * NeigborUVTransform.z;
-
 	// Calculate vertex UV details before morphing
 	float2 XY = ((float2)Pos + LocalUV) * (float)(1 << Level);
 	float2 NormalizedPos = (XY * PageTableSize.zw);
-	float2 LocalPhysicalUV = LocalUVTransform.xy + LocalUV * LocalUVTransform.z;
+	float SampleLevel = Level;
 
 	// Sample height once to approximate distance and morph the LocalUV.
 	{
+		float2 LocalPhysicalUV = LocalUVTransform.xy + LocalUV * LocalUVTransform.z;
 		float Height = HeightTexture.SampleLevel(HeightSampler, LocalPhysicalUV, 0);
-
 		float3 WorldPos = mul(float4(NormalizedPos, Height, 1), VirtualHeightfieldToWorld).xyz;
 		float DistanceSq = dot(LodViewOrigin - WorldPos, LodViewOrigin - WorldPos);
 		float Distance = sqrt(DistanceSq) * LodDistances.w;
@@ -113,19 +118,25 @@ FVertexFactoryIntermediates GetVertexFactoryIntermediates(FVertexFactoryInput In
 		// Final fractional morph has smoothstep to push transition region out to end of the LOD level.
 		float LodMorphFrac = smoothstep(0.7f, 0.98f, frac(LodForDistanceClamped));
 
+		// Use LOD level as sample level but bias to keep sample level in a good range.
+		SampleLevel = max(0, LodForDistanceClamped - 0.5f);
+
 		// Apply morph to vertex postion.
 		LocalUV = MorphVertex(LocalUV, GRID_SIZE - 1, (uint)LodMorphFloor, LodMorphFrac);
-
+		// Adjust other position variables to match the new morphed position.
 		XY = ((float2)Pos + LocalUV) * (float)(1 << Level);
 		NormalizedPos = (XY * PageTableSize.zw);
-		LocalPhysicalUV = LocalUVTransform.xy + LocalUV * LocalUVTransform.z;
-		NeighborPhysicalUV = NeigborUVTransform.xy + LocalUV * NeigborUVTransform.z;
 	}
 	
-	// Height
-	float LocalHeight = HeightTexture.SampleLevel(HeightSampler, LocalPhysicalUV, 0);
-	float NeighborHeight = HeightTexture.SampleLevel(HeightSampler, NeighborPhysicalUV, 0);
-	float Height = lerp(LocalHeight, NeighborHeight, NeighborWeight);
+	// Sample height from virtual texture.
+	VTUniform Uniform = VTUniform_Unpack(VTPackedUniform);
+	Uniform.vPageBorderSize -= .5f * PhysicalTextureSize.y; // Half texel offset hack
+	VTPageTableUniform PageTableUniform = VTPageTableUniform_Unpack(VTPackedPageTableUniform0, VTPackedPageTableUniform1);
+	VTPageTableResult VTResult0 = TextureLoadVirtualPageTableLevel(PageTableTexture, PageTableUniform, NormalizedPos, VTADDRESSMODE_CLAMP, VTADDRESSMODE_CLAMP, floor(SampleLevel));
+	float4 Height0 = TextureVirtualSampleLevel(HeightTexture, HeightSampler, VTResult0, 0, Uniform);
+	VTPageTableResult VTResult1 = TextureLoadVirtualPageTableLevel(PageTableTexture, PageTableUniform, NormalizedPos, VTADDRESSMODE_CLAMP, VTADDRESSMODE_CLAMP, ceil(SampleLevel));
+	float4 Height1 = TextureVirtualSampleLevel(HeightTexture, HeightSampler, VTResult1, 0, Uniform);
+	float Height = lerp(Height0.x, Height1.x, frac(SampleLevel));
 
 	// Position in space of virtual texture volume
 	Intermediates.VTPos = float3(NormalizedPos, Height);
diff --git a/Engine/Plugins/Experimental/VirtualHeightfieldMesh/Source/VirtualHeightfieldMesh/Private/VirtualHeightfieldMeshSceneProxy.cpp b/Engine/Plugins/Experimental/VirtualHeightfieldMesh/Source/VirtualHeightfieldMesh/Private/VirtualHeightfieldMeshSceneProxy.cpp
--- a/Engine/Plugins/Experimental/VirtualHeightfieldMesh/Source/VirtualHeightfieldMesh/Private/VirtualHeightfieldMeshSceneProxy.cpp
+++ b/Engine/Plugins/Experimental/VirtualHeightfieldMesh/Source/VirtualHeightfieldMesh/Private/VirtualHeightfieldMeshSceneProxy.cpp
@@ -523,12 +523,19 @@ void FVirtualHeightfieldMeshSceneProxy::GetDynamicMeshElements(const TArray<cons
 				BatchElement.UserData = (void*)UserData;
 
 				UserData->InstanceBufferSRV = Buffers.InstanceBufferSRV;
+				UserData->PageTableTexture = AllocatedVirtualTexture->GetPageTableTexture(0);
 				UserData->HeightPhysicalTexture = AllocatedVirtualTexture->GetPhysicalTexture(0);
 
+				AllocatedVirtualTexture->GetPackedUniform(&UserData->PackedUniform, 0);
+				AllocatedVirtualTexture->GetPackedPageTableUniform(UserData->PackedPageTableUniform);
+
 				const float PageTableSizeX = AllocatedVirtualTexture->GetWidthInTiles();
 				const float PageTableSizeY = AllocatedVirtualTexture->GetHeightInTiles();
 				UserData->PageTableSize = FVector4(PageTableSizeX, PageTableSizeY, 1.f / PageTableSizeX, 1.f / PageTableSizeY);
 
+				const float PhysicalTextureSize = AllocatedVirtualTexture->GetPhysicalTextureSize(0);
+				UserData->PhysicalTextureSize = FVector2D(PhysicalTextureSize, 1.f / PhysicalTextureSize);
+
 				UserData->MaxLod = AllocatedVirtualTexture->GetMaxLevel() + NumTailLods;
 				UserData->VirtualHeightfieldToLocal = UVToLocal;
 				UserData->VirtualHeightfieldToWorld = UVToWorld;
@@ -675,9 +682,8 @@ namespace VirtualHeightfieldMesh
 	/** Final render instance description used by the DrawInstancedIndirect(). Keep in sync with VirtualHeightfieldMesh.ush. */
 	struct QuadRenderInstance
 	{
-		uint32 AddressLevelPacked;
 		float UVTransform[3];
-		float NeigborUVTransform[4][4];
+		uint32 AddressLevelPacked;
 	};
 
 	/** Compute shader to initialize all buffers, including adding the lowest mip page(s) to the QuadBuffer. */
diff --git a/Engine/Plugins/Experimental/VirtualHeightfieldMesh/Source/VirtualHeightfieldMesh/Private/VirtualHeightfieldMeshVertexFactory.cpp b/Engine/Plugins/Experimental/VirtualHeightfieldMesh/Source/VirtualHeightfieldMesh/Private/VirtualHeightfieldMeshVertexFactory.cpp
--- a/Engine/Plugins/Experimental/VirtualHeightfieldMesh/Source/VirtualHeightfieldMesh/Private/VirtualHeightfieldMeshVertexFactory.cpp
+++ b/Engine/Plugins/Experimental/VirtualHeightfieldMesh/Source/VirtualHeightfieldMesh/Private/VirtualHeightfieldMeshVertexFactory.cpp
@@ -100,9 +100,14 @@ public:
 	void Bind(const FShaderParameterMap& ParameterMap)
 	{
 		InstanceBufferParameter.Bind(ParameterMap, TEXT("InstanceBuffer"));
+		PageTableTextureParameter.Bind(ParameterMap, TEXT("PageTableTexture"));
 		HeightTextureParameter.Bind(ParameterMap, TEXT("HeightTexture"));
 		HeightSamplerParameter.Bind(ParameterMap, TEXT("HeightSampler"));
+		PackedUniformParameter.Bind(ParameterMap, TEXT("VTPackedUniform"));
+		PackedPageTableUniform0Parameter.Bind(ParameterMap, TEXT("VTPackedPageTableUniform0"));
+		PackedPageTableUniform1Parameter.Bind(ParameterMap, TEXT("VTPackedPageTableUniform1"));
 		PageTableSizeParameter.Bind(ParameterMap, TEXT("PageTableSize"));
+		PhysicalTextureSizeParameter.Bind(ParameterMap, TEXT("PhysicalTextureSize"));
 		MaxLodParameter.Bind(ParameterMap, TEXT("MaxLod"));
 		VirtualHeightfieldToLocalParameter.Bind(ParameterMap, TEXT("VirtualHeightfieldToLocal"));
 		VirtualHeightfieldToWorldParameter.Bind(ParameterMap, TEXT("VirtualHeightfieldToWorld"));
@@ -132,14 +137,34 @@ public:
 			{
 				ShaderBindings.Add(InstanceBufferParameter, UserData->InstanceBufferSRV);
 			}
+			if (PageTableTextureParameter.IsBound() && PageTableTextureParameter.IsBound())
+			{
+				ShaderBindings.AddTexture(PageTableTextureParameter, PageTableTextureParameter, TStaticSamplerState<SF_Point>::GetRHI(), UserData->PageTableTexture);
+			}
 			if (HeightTextureParameter.IsBound() && HeightSamplerParameter.IsBound())
 			{
 				ShaderBindings.AddTexture(HeightTextureParameter, HeightSamplerParameter, TStaticSamplerState<SF_Bilinear>::GetRHI(), UserData->HeightPhysicalTexture);
 			}
+			if (PackedUniformParameter.IsBound())
+			{
+				ShaderBindings.Add(PackedUniformParameter, UserData->PackedUniform);
+			}
+			if (PackedPageTableUniform0Parameter.IsBound())
+			{
+				ShaderBindings.Add(PackedPageTableUniform0Parameter, UserData->PackedPageTableUniform[0]);
+			}
+			if (PackedPageTableUniform1Parameter.IsBound())
+			{
+				ShaderBindings.Add(PackedPageTableUniform1Parameter, UserData->PackedPageTableUniform[1]);
+			}
 			if (PageTableSizeParameter.IsBound())
 			{
 				ShaderBindings.Add(PageTableSizeParameter, UserData->PageTableSize);
 			}
+			if (PhysicalTextureSizeParameter.IsBound())
+			{
+				ShaderBindings.Add(PhysicalTextureSizeParameter, UserData->PhysicalTextureSize);
+			}
 			if (MaxLodParameter.IsBound())
 			{
 				ShaderBindings.Add(MaxLodParameter, UserData->MaxLod);
@@ -165,9 +190,14 @@ public:
 
 protected:
 	LAYOUT_FIELD(FShaderResourceParameter, InstanceBufferParameter);
+	LAYOUT_FIELD(FShaderResourceParameter, PageTableTextureParameter);
 	LAYOUT_FIELD(FShaderResourceParameter, HeightTextureParameter);
 	LAYOUT_FIELD(FShaderResourceParameter, HeightSamplerParameter);
+	LAYOUT_FIELD(FShaderParameter, PackedUniformParameter);
+	LAYOUT_FIELD(FShaderParameter, PackedPageTableUniform0Parameter);
+	LAYOUT_FIELD(FShaderParameter, PackedPageTableUniform1Parameter);
 	LAYOUT_FIELD(FShaderParameter, PageTableSizeParameter);
+	LAYOUT_FIELD(FShaderParameter, PhysicalTextureSizeParameter);
 	LAYOUT_FIELD(FShaderParameter, MaxLodParameter);
 	LAYOUT_FIELD(FShaderParameter, VirtualHeightfieldToLocalParameter);
 	LAYOUT_FIELD(FShaderParameter, VirtualHeightfieldToWorldParameter);
diff --git a/Engine/Plugins/Experimental/VirtualHeightfieldMesh/Source/VirtualHeightfieldMesh/Private/VirtualHeightfieldMeshVertexFactory.h b/Engine/Plugins/Experimental/VirtualHeightfieldMesh/Source/VirtualHeightfieldMesh/Private/VirtualHeightfieldMeshVertexFactory.h
--- a/Engine/Plugins/Experimental/VirtualHeightfieldMesh/Source/VirtualHeightfieldMesh/Private/VirtualHeightfieldMeshVertexFactory.h
+++ b/Engine/Plugins/Experimental/VirtualHeightfieldMesh/Source/VirtualHeightfieldMesh/Private/VirtualHeightfieldMeshVertexFactory.h
@@ -23,8 +23,12 @@ typedef TUniformBufferRef<FVirtualHeightfieldMeshVertexFactoryParameters> FVirtu
 struct FVirtualHeightfieldMeshUserData : public FOneFrameResource
 {
 	FRHIShaderResourceView* InstanceBufferSRV;
+	FRHITexture* PageTableTexture;
 	FRHITexture* HeightPhysicalTexture;
+	FUintVector4 PackedUniform;
+	FUintVector4 PackedPageTableUniform[2];
 	FVector4 PageTableSize;
+	FVector2D PhysicalTextureSize;
 	float MaxLod;
 	FMatrix VirtualHeightfieldToLocal;
 	FMatrix VirtualHeightfieldToWorld;
