Fix static analysis warning

diff --git a/Engine/Plugins/Experimental/VirtualHeightfieldMesh/Source/VirtualHeightfieldMesh/Private/VirtualHeightfieldMeshVertexFactory.cpp b/Engine/Plugins/Experimental/VirtualHeightfieldMesh/Source/VirtualHeightfieldMesh/Private/VirtualHeightfieldMeshVertexFactory.cpp
--- a/Engine/Plugins/Experimental/VirtualHeightfieldMesh/Source/VirtualHeightfieldMesh/Private/VirtualHeightfieldMeshVertexFactory.cpp
+++ b/Engine/Plugins/Experimental/VirtualHeightfieldMesh/Source/VirtualHeightfieldMesh/Private/VirtualHeightfieldMeshVertexFactory.cpp
@@ -137,9 +137,9 @@ public:
 			{
 				ShaderBindings.Add(InstanceBufferParameter, UserData->InstanceBufferSRV);
 			}
-			if (PageTableTextureParameter.IsBound() && PageTableTextureParameter.IsBound())
+			if (PageTableTextureParameter.IsBound())
 			{
-				ShaderBindings.AddTexture(PageTableTextureParameter, PageTableTextureParameter, TStaticSamplerState<SF_Point>::GetRHI(), UserData->PageTableTexture);
+				ShaderBindings.AddTexture(PageTableTextureParameter, PageTableSamplerParameter, TStaticSamplerState<SF_Point>::GetRHI(), UserData->PageTableTexture);
 			}
 			if (HeightTextureParameter.IsBound() && HeightSamplerParameter.IsBound())
 			{
@@ -191,6 +191,7 @@ public:
 protected:
 	LAYOUT_FIELD(FShaderResourceParameter, InstanceBufferParameter);
 	LAYOUT_FIELD(FShaderResourceParameter, PageTableTextureParameter);
+	LAYOUT_FIELD(FShaderResourceParameter, PageTableSamplerParameter);
 	LAYOUT_FIELD(FShaderResourceParameter, HeightTextureParameter);
 	LAYOUT_FIELD(FShaderResourceParameter, HeightSamplerParameter);
 	LAYOUT_FIELD(FShaderParameter, PackedUniformParameter);
