Fix Height MinMax calculation in VHM to ignore unwritten pixels.
Fixes issue with oversized bounding boxes and LodBias texture generation.

diff --git a/Engine/Plugins/Experimental/VirtualHeightfieldMesh/Shaders/Private/HeightfieldMinMaxRender.usf b/Engine/Plugins/Experimental/VirtualHeightfieldMesh/Shaders/Private/HeightfieldMinMaxRender.usf
--- a/Engine/Plugins/Experimental/VirtualHeightfieldMesh/Shaders/Private/HeightfieldMinMaxRender.usf
+++ b/Engine/Plugins/Experimental/VirtualHeightfieldMesh/Shaders/Private/HeightfieldMinMaxRender.usf
@@ -28,6 +28,18 @@ float2 UnPackMinMax(in float4 Packed)
 	return UnPacked;
 }
 
+// Min downsample excludes zero values to avoid empty pixels affecting the calculation
+float Min4(in float4 Values, in float ReplaceZero)
+{
+	Values = Values == 0 ? ReplaceZero : Values;
+	return min(min(min(Values.x, Values.y), Values.z), Values.w);
+}
+
+float Max4(in float4 Values)
+{
+	return max(max(max(Values.x, Values.y), Values.z), Values.w);
+}
+
 [numthreads(8, 8, 1)]
 void MinMaxHeightCS(uint3 DispatchThreadId : SV_DispatchThreadID)
 {
@@ -39,22 +51,22 @@ void MinMaxHeightCS(uint3 DispatchThreadId : SV_DispatchThreadID)
 	float T01 = SrcTexture[2 * DispatchThreadId.xy + int2(0, 1)].x;
 	float T10 = SrcTexture[2 * DispatchThreadId.xy + int2(1, 0)].x;
 	float T11 = SrcTexture[2 * DispatchThreadId.xy + int2(1, 1)].x;
-	float ValueMin = min(min(min(T00, T01), T10), T11);
-	float ValueMax = max(max(max(T00, T01), T10), T11);
+	float ValueMax = Max4(float4(T00, T01, T10, T11));
+	float ValueMin = Min4(float4(T00, T01, T10, T11), ValueMax);
 #elif INPUT_FORMAT_RG16
  	float2 T00 = SrcTexture[2 * DispatchThreadId.xy].xy;
  	float2 T01 = SrcTexture[2 * DispatchThreadId.xy + int2(0, 1)].xy;
  	float2 T10 = SrcTexture[2 * DispatchThreadId.xy + int2(1, 0)].xy;
  	float2 T11 = SrcTexture[2 * DispatchThreadId.xy + int2(1, 1)].xy;
- 	float ValueMin = min(min(min(T00.x, T01.x), T10.x), T11.x);
- 	float ValueMax = max(max(max(T00.y, T01.y), T10.y), T11.y);
+	float ValueMax = Max4(float4(T00.y, T01.y, T10.y, T11.y));
+	float ValueMin = Min4(float4(T00.x, T01.x, T10.x, T11.x), ValueMax);
 #elif INPUT_FORMAT_RGBA8
 	float2 T00 = UnPackMinMax(SrcTexture[2 * DispatchThreadId.xy]);
 	float2 T01 = UnPackMinMax(SrcTexture[2 * DispatchThreadId.xy + int2(0, 1)]);
 	float2 T10 = UnPackMinMax(SrcTexture[2 * DispatchThreadId.xy + int2(1, 0)]);
 	float2 T11 = UnPackMinMax(SrcTexture[2 * DispatchThreadId.xy + int2(1, 1)]);
-	float ValueMin = min(min(min(T00.x, T01.x), T10.x), T11.x);
-	float ValueMax = max(max(max(T00.y, T01.y), T10.y), T11.y);
+	float ValueMax = Max4(float4(T00.y, T01.y, T10.y, T11.y));
+	float ValueMin = Min4(float4(T00.x, T01.x, T10.x, T11.x), ValueMax);
 #endif
 
 #if OUTPUT_FORMAT_RG16
