Generate LodBias textures as part of VHM MinMaxHeight texture generation.

diff --git a/Engine/Plugins/Experimental/VirtualHeightfieldMesh/Source/VirtualHeightfieldMesh/Private/HeightfieldMinMaxTexture.cpp b/Engine/Plugins/Experimental/VirtualHeightfieldMesh/Source/VirtualHeightfieldMesh/Private/HeightfieldMinMaxTexture.cpp
--- a/Engine/Plugins/Experimental/VirtualHeightfieldMesh/Source/VirtualHeightfieldMesh/Private/HeightfieldMinMaxTexture.cpp
+++ b/Engine/Plugins/Experimental/VirtualHeightfieldMesh/Source/VirtualHeightfieldMesh/Private/HeightfieldMinMaxTexture.cpp
@@ -8,7 +8,6 @@
 
 UHeightfieldMinMaxTexture::UHeightfieldMinMaxTexture(const FObjectInitializer& ObjectInitializer)
 	: UObject(ObjectInitializer)
-	, Texture(nullptr)
 	, MaxCPULevels(5)
 {
 }
@@ -30,8 +29,17 @@ void UHeightfieldMinMaxTexture::PostEditChangeProperty(FPropertyChangedEvent& Pr
 
 void UHeightfieldMinMaxTexture::BuildTexture(FHeightfieldMinMaxTextureBuildDesc const& InBuildDesc)
 {
-	// Build GPU Texture.
-	Texture = NewObject<UTexture2D>(this, TEXT("Texture"));
+	RebuildMinMaxTexture(InBuildDesc);
+	RebuildLodBiasTexture(InBuildDesc);
+	RebuildCPUTextureData();
+
+	// Notify all dependent components.
+	VirtualHeightfieldMesh::NotifyComponents(this);
+}
+
+void UHeightfieldMinMaxTexture::RebuildMinMaxTexture(FHeightfieldMinMaxTextureBuildDesc const& InBuildDesc)
+{
+	Texture = NewObject<UTexture2D>(this, TEXT("HeightMinMaxTexture"));
 
 	FTextureFormatSettings Settings;
 	Settings.CompressionSettings = TC_EditorIcon;
@@ -46,12 +54,195 @@ void UHeightfieldMinMaxTexture::BuildTexture(FHeightfieldMinMaxTextureBuildDesc
 	Texture->Source.Init(InBuildDesc.SizeX, InBuildDesc.SizeY, 1, InBuildDesc.NumMips, TSF_BGRA8, InBuildDesc.Data);
 
 	Texture->PostEditChange();
-	
-	// Build CPU TextureData.
-	RebuildCPUTextureData();
+}
 
-	// Notify all dependent components.
-	VirtualHeightfieldMesh::NotifyComponents(this);
+void UHeightfieldMinMaxTexture::RebuildLodBiasTexture(FHeightfieldMinMaxTextureBuildDesc const& InBuildDesc)
+{
+	TArray<uint8> LodBiasTextureData;
+	TArray<uint32> MinMaxLodBiasTextureData;
+
+	// Build the LodBias texture data
+	{
+		// Allocate two work buffers (we will need second for blur passes).
+		TArray<uint16> ScratchBuffers[2];
+		ScratchBuffers[0].AddDefaulted(InBuildDesc.SizeX * InBuildDesc.SizeY);
+		ScratchBuffers[1].AddDefaulted(InBuildDesc.SizeX * InBuildDesc.SizeY);
+		int32 ScratchBufferWriteIndex = 0;
+		
+		// Extract MinMax values from source data and store the differences.
+		{
+			uint16 MaxHeightDiff = 1;
+			for (uint32 Index = 0; Index < InBuildDesc.SizeX * InBuildDesc.SizeY; ++Index)
+			{
+				const uint16 Min = InBuildDesc.Data[Index * 4 + 2] * 256 + InBuildDesc.Data[Index * 4 + 3];
+				const uint16 Max = InBuildDesc.Data[Index * 4 + 0] * 256 + InBuildDesc.Data[Index * 4 + 1];
+				const uint16 HeightDiff = FMath::Max(Max, Min) - Min;
+				ScratchBuffers[ScratchBufferWriteIndex][Index] = HeightDiff;
+				MaxHeightDiff = FMath::Max(MaxHeightDiff, HeightDiff);
+			}
+
+			// Normalize the data against the maximum value.
+			for (uint32 Index = 0; Index < InBuildDesc.SizeX * InBuildDesc.SizeY; ++Index)
+			{
+				const uint16 NormalizedRange = (uint16)(65535.f * ((float)ScratchBuffers[0][Index] / (float)MaxHeightDiff));
+				ScratchBuffers[ScratchBufferWriteIndex][Index] = NormalizedRange;
+			}
+		}
+
+		// Ping-pong blur to reduce maximum difference between any two texels.
+		// This is to remove sharp gradients in the data which will generate lod transitions that are too aggressive for the lod algorithm to stitch correctly.
+		{
+			const uint32 NumBlurPasses = 20;
+			const uint16 MaxGradient = 65535 / 10;
+
+			for (uint32 BlurIndex = 0; BlurIndex < NumBlurPasses; ++BlurIndex)
+			{
+				const int32 ScratchBufferReadIndex = ScratchBufferWriteIndex;
+				ScratchBufferWriteIndex = 1 - ScratchBufferWriteIndex;
+				
+				for (int32 Y = 0; Y < (int32)InBuildDesc.SizeY; ++Y)
+				{
+					for (int32 X = 0; X < (int32)InBuildDesc.SizeX; ++X)
+					{
+						// Get maximum value in the 3x3 neighborhood.
+						uint16 MaxValue = 0;
+						for (int32 YOffset = -1; YOffset <= 1; ++YOffset)
+						{
+							int32 YCoord = FMath::Clamp<int32>(Y + YOffset, 0, InBuildDesc.SizeY - 1);
+							for (int32 XOffset = -1; XOffset <= 1; ++XOffset)
+							{
+								int32 XCoord = FMath::Clamp<int32>(X + XOffset, 0, InBuildDesc.SizeX - 1);
+
+								MaxValue = FMath::Max(MaxValue, ScratchBuffers[ScratchBufferReadIndex][InBuildDesc.SizeY * YCoord + XCoord]);
+							}
+						}
+
+						// Clamp to within MaxGradient of MaxValue.
+						const uint16 MinValue = MaxValue > MaxGradient ? MaxValue - MaxGradient : 0;
+						const uint16 CurrentValue = ScratchBuffers[ScratchBufferReadIndex][InBuildDesc.SizeY * Y + X];
+						const uint16 NewValue = FMath::Max(CurrentValue, MinValue);
+						ScratchBuffers[ScratchBufferWriteIndex][InBuildDesc.SizeY * Y + X] = NewValue;
+					}
+				}
+			}
+		}
+
+		// Copy final ping-pong buffer to result data.
+		{
+			const int32 ScratchBufferReadIndex = ScratchBufferWriteIndex;
+
+			LodBiasTextureData.AddDefaulted(InBuildDesc.SizeX * InBuildDesc.SizeY);
+			for (uint32 Index = 0; Index < InBuildDesc.SizeX * InBuildDesc.SizeY; ++Index)
+			{
+				LodBiasTextureData[Index] = ScratchBuffers[ScratchBufferReadIndex][Index] >> 8;
+			}
+		}
+	}
+
+	// Create new texture
+	{
+		LodBiasTexture = NewObject<UTexture2D>(this, TEXT("LodBiasTexture"));
+
+		FTextureFormatSettings Settings;
+		Settings.CompressionSettings = TC_EditorIcon;
+		Settings.CompressionNone = true;
+		Settings.SRGB = false;
+
+		LodBiasTexture->Filter = TF_Nearest;
+		LodBiasTexture->MipGenSettings = TMGS_NoMipmaps;
+		LodBiasTexture->SetLayerFormatSettings(0, Settings);
+		LodBiasTexture->Source.Init(InBuildDesc.SizeX, InBuildDesc.SizeY, 1, InBuildDesc.NumMips, TSF_G8, LodBiasTextureData.GetData());
+
+		LodBiasTexture->PostEditChange();
+	}
+
+	// Build top mip of MinMaxLodBias texture. 
+	// Include a texel border. This will account for any bleed in the vertex factory code where we bilinear sample the LodBiasTexture.
+	{
+		// Reserve the expected entries assuming square mips. This may be an overestimate.
+		MinMaxLodBiasTextureData.Reserve(((1 << (2 * InBuildDesc.NumMips)) - 1) / 3);
+
+		for (int32 Y = 0; Y < (int32)InBuildDesc.SizeY; ++Y)
+		{
+			for (int32 X = 0; X < (int32)InBuildDesc.SizeX; ++X)
+			{
+				uint32 MaxValue = 0;
+				uint32 MinValue = 255;
+				for (int32 YOffset = -1; YOffset <= 1; ++YOffset)
+				{
+					const int32 YCoord = FMath::Clamp<int32>(Y + YOffset, 0, InBuildDesc.SizeY - 1);
+					for (int32 XOffset = -1; XOffset <= 1; ++XOffset)
+					{
+						const int32 XCoord = FMath::Clamp<int32>(X + XOffset, 0, InBuildDesc.SizeX - 1);
+						const uint32 LodBiasValue = LodBiasTextureData[InBuildDesc.SizeY * YCoord + XCoord];
+
+						MaxValue = FMath::Max(MaxValue, LodBiasValue);
+						MinValue = FMath::Min(MinValue, LodBiasValue);
+					}
+				}
+
+				// We pack min and max values to map to the red and green channels in our BGRA8 texture.
+				// It would be better to use an RG8 here but that isn't supported by texture import.
+				MinMaxLodBiasTextureData.Add((MaxValue << 8) | (MinValue << 16) | 0xFF000000);
+			}
+		}
+	}
+
+	// Build mips with MinMax filter
+	{
+		int32 SizeX = InBuildDesc.SizeX;
+		int32 SizeY = InBuildDesc.SizeY;
+		uint32 const* SourceData = MinMaxLodBiasTextureData.GetData();
+		for (uint32 MipIndex = 1; MipIndex < InBuildDesc.NumMips; ++MipIndex)
+		{
+			const int32 SourceSizeX = SizeX;
+			const int32 SourceSizeY = SizeY;
+			
+			SizeX = FMath::Max(SizeX >> 1, 1);
+			SizeY = FMath::Max(SizeY >> 1, 1);
+
+			for (int32 Y = 0; Y < SizeY; ++Y)
+			{
+				for (int32 X = 0; X < SizeX; ++X)
+				{
+					const int32 X0 = FMath::Min(X * 2, SourceSizeX - 1);
+					const int32 X1 = FMath::Min(X * 2 + 1, SourceSizeX - 1);
+					const int32 Y0 = FMath::Min(Y * 2, SourceSizeY - 1);
+					const int32 Y1 = FMath::Min(Y * 2 + 1, SourceSizeY - 1);
+					const uint32 S00 = *(SourceData + Y0 * SourceSizeY + X0);
+					const uint32 S10 = *(SourceData + Y0 * SourceSizeY + X1);
+					const uint32 S01 = *(SourceData + Y1 * SourceSizeY + X0);
+					const uint32 S11 = *(SourceData + Y1 * SourceSizeY + X1);
+
+					// MinMax filter assumes BGRA8 layout.
+					const uint32 MaxValue = FMath::Max(FMath::Max(FMath::Max((S00 & 0x00FF00), (S10 & 0x00FF00)), (S01 & 0x00FF00)), (S11 & 0x00FF00));
+					const uint32 MinValue = FMath::Min(FMath::Min(FMath::Min((S00 & 0xFF0000), (S10 & 0xFF0000)), (S01 & 0xFF0000)), (S11 & 0xFF0000));
+					MinMaxLodBiasTextureData.Add(MaxValue | MinValue | 0xFF000000);
+				}
+			}
+
+			SourceData += SourceSizeX * SourceSizeY;
+		}
+	}
+
+	// Create new texture
+	{
+		LodBiasMinMaxTexture = NewObject<UTexture2D>(this, TEXT("LodBiasMinMaxTexture"));
+
+		FTextureFormatSettings Settings;
+		Settings.CompressionSettings = TC_EditorIcon;
+		Settings.CompressionNone = true;
+		Settings.SRGB = false;
+
+		LodBiasMinMaxTexture->Filter = TF_Nearest;
+		LodBiasMinMaxTexture->MipGenSettings = TMGS_LeaveExistingMips;
+		LodBiasMinMaxTexture->MipLoadOptions = ETextureMipLoadOptions::AllMips;
+		LodBiasMinMaxTexture->NeverStream = true;
+		LodBiasMinMaxTexture->SetLayerFormatSettings(0, Settings);
+		LodBiasMinMaxTexture->Source.Init(InBuildDesc.SizeX, InBuildDesc.SizeY, 1, InBuildDesc.NumMips, TSF_BGRA8, (uint8*)MinMaxLodBiasTextureData.GetData());
+
+		LodBiasMinMaxTexture->PostEditChange();
+	}
 }
 
 void UHeightfieldMinMaxTexture::RebuildCPUTextureData()
diff --git a/Engine/Plugins/Experimental/VirtualHeightfieldMesh/Source/VirtualHeightfieldMesh/Public/HeightfieldMinMaxTexture.h b/Engine/Plugins/Experimental/VirtualHeightfieldMesh/Source/VirtualHeightfieldMesh/Public/HeightfieldMinMaxTexture.h
--- a/Engine/Plugins/Experimental/VirtualHeightfieldMesh/Source/VirtualHeightfieldMesh/Public/HeightfieldMinMaxTexture.h
+++ b/Engine/Plugins/Experimental/VirtualHeightfieldMesh/Source/VirtualHeightfieldMesh/Public/HeightfieldMinMaxTexture.h
@@ -28,9 +28,17 @@ class VIRTUALHEIGHTFIELDMESH_API UHeightfieldMinMaxTexture : public UObject
 public:
 	GENERATED_UCLASS_BODY()
 
-	/** The UTexture object. */
-	UPROPERTY(VisibleAnywhere, BlueprintReadOnly, Category = Texture)
-	class UTexture2D* Texture;
+	/** The Height MinMax texture. */
+	UPROPERTY(VisibleAnywhere, BlueprintReadOnly, Category = Texture, meta = (DisplayName = "Height MinMax Texture"))
+	class UTexture2D* Texture = nullptr;
+
+	/** A LodBias texture derived from the Height MinMax texture*/
+	UPROPERTY(VisibleAnywhere, BlueprintReadOnly, Category = Texture, meta = (DisplayName = "LodBias Texture"))
+	class UTexture2D* LodBiasTexture = nullptr;
+
+	/** A LodBias MinMax texture derived from the LodBias texture*/
+	UPROPERTY(VisibleAnywhere, BlueprintReadOnly, Category = Texture, meta = (DisplayName = "LodBias MinMax Texture"))
+	class UTexture2D* LodBiasMinMaxTexture = nullptr;
 
 protected:
 	/** The number of mip levels to clone for CPU access. */
@@ -51,12 +59,16 @@ public:
 	TArray<int32> TextureDataMips;
 
 #if WITH_EDITOR
-	/** Creates a new UVirtualTexture2D and stores it in the contained Texture. */
+	/** Creates a new UTexture2D and stores it in the contained Texture. */
 	void BuildTexture(FHeightfieldMinMaxTextureBuildDesc const& InBuildDesc);
 #endif
  
 protected:
 #if WITH_EDITOR
+	/** Rebuild the contents of Texture using the passed in build description. */
+	void RebuildMinMaxTexture(FHeightfieldMinMaxTextureBuildDesc const& InBuildDesc);
+	/** Rebuild the contents of LodBiasTexture from the MinMaxTexture. */
+	void RebuildLodBiasTexture(FHeightfieldMinMaxTextureBuildDesc const& InBuildDesc);
 	/** Rebuild the contents of TextureData for the specified number of CPU mip levels. */
 	void RebuildCPUTextureData();
 
diff --git a/Engine/Plugins/Experimental/VirtualHeightfieldMesh/Source/VirtualHeightfieldMeshEditor/Private/HeightfieldMinMaxTextureBuild.cpp b/Engine/Plugins/Experimental/VirtualHeightfieldMesh/Source/VirtualHeightfieldMeshEditor/Private/HeightfieldMinMaxTextureBuild.cpp
--- a/Engine/Plugins/Experimental/VirtualHeightfieldMesh/Source/VirtualHeightfieldMeshEditor/Private/HeightfieldMinMaxTextureBuild.cpp
+++ b/Engine/Plugins/Experimental/VirtualHeightfieldMesh/Source/VirtualHeightfieldMeshEditor/Private/HeightfieldMinMaxTextureBuild.cpp
@@ -159,7 +159,7 @@ namespace VirtualHeightfieldMesh
 		FScopedSlowTask Task(TaskWorkRender + TaskWorkDownsample + TaskWorkBuildBulkData, FText::AsCultureInvariant(InComponent->GetMinMaxTexture()->GetName()));
 		Task.MakeDialog(true);
 
-		// Final pixels will contain image data for each virtual texture layer in order
+		// Final pixels will contain image data for MinMax texture
 		TArray64<uint8> FinalPixels;
 		FinalPixels.SetNumUninitialized(RenderTileResources.GetNumFinalTexels() * 4);
 
