Apply auto lod bias in VHM rendering

diff --git a/Engine/Plugins/Experimental/VirtualHeightfieldMesh/Shaders/Private/VirtualHeightfieldMesh.usf b/Engine/Plugins/Experimental/VirtualHeightfieldMesh/Shaders/Private/VirtualHeightfieldMesh.usf
--- a/Engine/Plugins/Experimental/VirtualHeightfieldMesh/Shaders/Private/VirtualHeightfieldMesh.usf
+++ b/Engine/Plugins/Experimental/VirtualHeightfieldMesh/Shaders/Private/VirtualHeightfieldMesh.usf
@@ -18,9 +18,6 @@ Buffer<uint2> QuadBuffer;
 RWBuffer<uint> RWIndirectArgsBuffer;
 Buffer<uint> IndirectArgsBufferSRV;
 
-RWBuffer<uint> RWQuadNeighborBuffer;
-Buffer<uint> QuadNeighborBuffer;
-
 RWStructuredBuffer<QuadRenderInstance> RWInstanceBuffer;
 StructuredBuffer<QuadRenderInstance> InstanceBuffer;
 
@@ -28,15 +25,14 @@ RWBuffer<uint> RWFeedbackBuffer;
 
 Texture2D<uint> PageTableTexture;
 
-Texture2D<float4> MinMaxTexture;
+Texture2D<float4> HeightMinMaxTexture;
+Texture2D<float4> LodBiasMinMaxTexture;
 SamplerState MinMaxTextureSampler;
 int MinMaxLevelOffset;
 
 Texture2D<float> OcclusionTexture;
 int OcclusionLevelOffset;
 
-Texture2D<float> LodTexture;
-
 uint MaxLevel;
 uint NumForceLoadLods;
 uint PageTableFeedbackId;
@@ -45,6 +41,7 @@ float4 PageTableSize;
 float4 PhysicalPageTransform;
 
 float4 LodDistances;
+float LodBiasScale;
 float3 ViewOrigin;
 float4 FrustumPlanes[5];
 float4x4 UVToWorld;
@@ -110,6 +107,12 @@ float2 UnPackMinMaxHeight(float4 InPacked)
 	return UnPacked;
 }
 
+/** Unpack the values from the MinMaxLodBias texture from the packed 8888 format. */
+float2 UnPackMinMaxLodBias(float4 InPacked, float InLodBiasScale)
+{
+	return float2(CalculateBiasLod(InPacked.x, InLodBiasScale), CalculateBiasLod(InPacked.y, InLodBiasScale));
+}
+
 /** Return false if the AABB is completely outside one of the planes. */
 bool PlaneTestAABB(float4 InPlanes[5], float3 InCenter, float3 InExtent)
 {
@@ -266,7 +269,8 @@ void CollectQuadsCS(
 			float2 UV1 = ((float2)Pos + float2(1, 1)) * Scale;
 			
 			float MinMaxTextureLevel = max((float)Level + (float)MinMaxLevelOffset, 0);
-			float2 MinMaxHeight = UnPackMinMaxHeight(MinMaxTexture.SampleLevel(MinMaxTextureSampler, UV0, MinMaxTextureLevel));
+			float2 MinMaxHeight = UnPackMinMaxHeight(HeightMinMaxTexture.SampleLevel(MinMaxTextureSampler, UV0, MinMaxTextureLevel));
+			float2 MinMaxLodBias = UnPackMinMaxLodBias(LodBiasMinMaxTexture.SampleLevel(MinMaxTextureSampler, UV0, MinMaxTextureLevel), LodBiasScale);
 			
 			float3 UVMin = float3(UV0, MinMaxHeight.x);
 			float3 UVMax = float3(UV1, MinMaxHeight.y);
@@ -274,6 +278,11 @@ void CollectQuadsCS(
 			float3 UVCenter = (UVMax + UVMin) * 0.5f;
 			float3 UVExtent = UVMax - UVCenter;
 
+			float MinDistanceSq = SquaredMinDistanceToAABB(ViewOrigin, UVMin, UVMax, UVToWorldScale);
+			float MaxDistanceSq = SquaredMaxDistanceToAABB(ViewOrigin, UVMin, UVMax, UVToWorldScale);
+			float MinDistanceLod = CalculateDistanceLod(MinDistanceSq, LodDistances);
+			float MaxDistanceLod = CalculateDistanceLod(MaxDistanceSq, LodDistances);
+
 			// Check if frustum culled
 			bool bCull = !PlaneTestAABB(FrustumPlanes, UVCenter, UVExtent);
 
@@ -285,12 +294,8 @@ void CollectQuadsCS(
 			}
 			else if (Level > 0)
 			{
-				// Subdivide if any part of bbox could be inside the LOD range.
-				float DistanceSq = SquaredMinDistanceToAABB(ViewOrigin, UVMin, UVMax, UVToWorldScale) * LodDistances.w * LodDistances.w;
-				float LodDistance = LodDistances.x * (LodDistances.y - 1 + pow(LodDistances.z, Level - 1));
-				float LodDistanceSq = LodDistance * LodDistance;
-
-				bSubdivide = DistanceSq < LodDistanceSq;
+				// Subdivide if minimum continuous lod can be less than the current level.
+				bSubdivide = MinDistanceLod - MinMaxLodBias.y < (float)Level;
 			}
 
 			if (bSubdivide)
@@ -317,9 +322,25 @@ void CollectQuadsCS(
 
 				RWQuadBuffer[Write] = Pack(InitQuadRenderItem(Pos, Level, PhysicalAddress, bCull || bOcclude));
 
+				// Add all possible pages that vertex shader might read to the virtual texture feedback buffer.
+				int MinFeedbackLevel = (int)floor(clamp(MinDistanceLod - MinMaxLodBias.y, Level, MaxLevel));
+				int MaxFeedbackLevel = (int)ceil(clamp(MaxDistanceLod - MinMaxLodBias.x, Level, MaxLevel));
+				uint NumFeedbackItems = (uint)max(MaxFeedbackLevel - MinFeedbackLevel + 1, 0);
+
+				uint FeedbackPos;
+				InterlockedAdd(RWFeedbackBuffer[0], NumFeedbackItems, FeedbackPos);
+				
+				for (int FeedbackLevel = MinFeedbackLevel; FeedbackLevel <= MaxFeedbackLevel; ++FeedbackLevel)
+				{
+					// Note that our general virtual texture feedback buffer convention is to write Level+1
+					uint LevelPlusOne = FeedbackLevel + 1;
+					uint LodShift = FeedbackLevel - Level;
+					RWFeedbackBuffer[FeedbackPos + FeedbackLevel] = (Pos.x >> LodShift) | ((Pos.y >> LodShift) << 12) | (LevelPlusOne << 24) | PageTableFeedbackId;
+				}
+
+				// Debug draw the bounds.
 				if (!(bCull || bOcclude))
 				{
-					// Debug draw the bounds.
 					DebugDrawUVBox(UVMin, UVMax, UVToWorld, float4(1, 0, 0, 1));
 				}
 			}
@@ -334,114 +355,6 @@ void CollectQuadsCS(
 	}
 }
 
-/**
- * Vertex/Pixel shaders that draw the Lod info for all quads output by the Collect pass.
- */
-
-void RenderLodMapVS(
-	in uint InstanceId : SV_InstanceID,
-	in uint VertexId : SV_VertexID,
-	out float4 OutPosition : SV_POSITION,
-	out float OutLod : TEXCOORD0 )
-{
-	uint2 PackedItem = QuadBuffer[InstanceId];
-	QuadRenderItem Item = UnpackQuadRenderItem(PackedItem);
-	
-	// Calculate vertex position for this quad.
-	uint2 Pos = Item.Pos;
-	uint Level = Item.Level;
-	uint2 Offset = uint2((VertexId >> 0) & 1, (VertexId >> 1) & 1);
-	uint2 VertexPosLod0 = (Pos + Offset) << Level;
-	float2 NormalizedVertexPos = (float2)VertexPosLod0 * PageTableSize.zw;
-	
-	OutPosition.xy = NormalizedVertexPos * float2(2.f, -2.f) + float2(-1.f, 1.f);
-	OutPosition.zw = float2(0, 1);
-
-	// Calculate currently streamed level in the virtual texture for this quad.
-	uint Lod = GetVirtualLevelFromPhysicalAddress(Item.PhysicalAddress);
-	
-	// Pack along with culling state.
-	uint CullBit = Item.bCull ? 1 : 0;
-	float PackedLod = (float)(Lod | (CullBit << 7)) / 255.f;
-
-	OutLod = PackedLod;
-}
-
-void RenderLodMapPS(
-	in float4 InPosition : SV_POSITION,
-	in noperspective float InLod : TEXCOORD0,
-	out float4 OutColor : SV_Target0 )
-{
-	OutColor = float4(InLod, 0.f, 0.f, 1.f);
-}
-
-/**
- * Read the Lod info for the neighbors of each quad and resolve Lods at borders.
- */
-
-[numthreads(64, 4, 1)]
-void ResolveNeighborLodsCS( 
-	uint GroupId : SV_GroupID,
-	uint GroupIndex : SV_GroupIndex )
-{
-	// Process one neighbor per thread.
-	uint ThreadId = GroupId * 256 + GroupIndex;
-	uint QuadIndex = ThreadId >> 2;
-	uint NeighborIndex = ThreadId & 3;
-
-	if (QuadIndex >= IndirectArgsBufferSRV[1])
-		return;
-
-	uint2 PackedItem = QuadBuffer[QuadIndex];
-	QuadRenderItem Item = UnpackQuadRenderItem(PackedItem);
-	uint2 Pos = Item.Pos;
-	uint Level = Item.Level;
-	uint PhysicalAddress = Item.PhysicalAddress;
-
-	// Get sample location of neighbor in the LodTexture.
-	int2 SampleOffset = 0;
-	SampleOffset += (NeighborIndex == 0) ? int2(-1, 0) : 0;
-	SampleOffset += (NeighborIndex == 1) ? int2(0, (1 << Level)) : 0;
-	SampleOffset += (NeighborIndex == 2) ? int2((1 << Level), 0) : 0;
-	SampleOffset += (NeighborIndex == 3) ? int2(0, -1) : 0;
-
-	uint2 PosLod0 = Pos << Level;
-	int2 ClampedPosLod0 = clamp((int2)PosLod0 + SampleOffset, 0, (int2)PageTableSize.xy - 1);
-
-	// Read the neighbor Lod information.
-	float NeighborLodSample = LodTexture.Load(int3(ClampedPosLod0, 0));
-	uint NeighborLodPacked = (uint)floor(NeighborLodSample * 255.f);
-	uint NeighborLod = NeighborLodPacked & 0x7f;
-	bool bNeighborActive = (NeighborLodPacked >> 7) == 0;
-	
-	// Get the PhysicalAddress for the local location but sampling at the neighbor Lod.
-	// Note that this will be the same as the local PhysicalAddress if the Lods are the same.
-	uint LocalLod = GetVirtualLevelFromPhysicalAddress(PhysicalAddress);
-	uint LodShift = max(NeighborLod - LocalLod, 0u);
-	uint NeighborLodPhysicalAddress = PageTableTexture.Load(int3(Pos >> LodShift, Level + LodShift));
-
-	// We will need to blend if we are transitioning to an on-screen lower detail Lod.
-	bool bNeighborLerp = bNeighborActive && NeighborLod > LocalLod;
-	uint PhysicalAddressToWrite = bNeighborLerp ? NeighborLodPhysicalAddress : PhysicalAddress;
-
-	// Write the results to a flat buffer with neighbor entries per render item.
-	RWQuadNeighborBuffer[ThreadId] = Pack(InitQuadNeighborItem(PhysicalAddressToWrite, bNeighborLerp));
-
-
-	// Add requests to the virtual texture feedback buffer.
-	uint NumFeedbackItems = (NeighborIndex == 0) ? 2 : 1;
-	uint FeedbackPos;
-	InterlockedAdd(RWFeedbackBuffer[0], NumFeedbackItems, FeedbackPos);
-
-	// Note that our general virtual texture feedback buffer convention is to write Level+1
-	uint LevelPlusOne = Level + 1;
-	RWFeedbackBuffer[FeedbackPos + 1] = (Pos.x >> LodShift) | ((Pos.y >> LodShift) << 12) | ((LevelPlusOne + LodShift) << 24) | PageTableFeedbackId;
-	if (NeighborIndex == 0)
-	{
-		RWFeedbackBuffer[FeedbackPos + 2] = Pos.x | (Pos.y << 12) | (LevelPlusOne << 24) | PageTableFeedbackId;
-	}
-}
-
 /**
  * Initialise the indirect args for the final culled indirect draw call.
  */
@@ -481,7 +394,8 @@ void CullInstancesCS( uint3 DispatchThreadId : SV_DispatchThreadID )
 	float2 UV0 = ((float2)Pos + float2(0, 0)) * Scale;
 	float2 UV1 = ((float2)Pos + float2(1, 1)) * Scale;
 
-	float2 MinMaxHeight = UnPackMinMaxHeight(MinMaxTexture.SampleLevel(MinMaxTextureSampler, UV0, (float)Level));
+	float MinMaxTextureLevel = max((float)Level + (float)MinMaxLevelOffset, 0);
+	float2 MinMaxHeight = UnPackMinMaxHeight(HeightMinMaxTexture.SampleLevel(MinMaxTextureSampler, UV0, MinMaxTextureLevel));
 
 	float3 UVMin = float3(UV0, MinMaxHeight.x);
 	float3 UVMax = float3(UV1, MinMaxHeight.y);
diff --git a/Engine/Plugins/Experimental/VirtualHeightfieldMesh/Shaders/Private/VirtualHeightfieldMesh.ush b/Engine/Plugins/Experimental/VirtualHeightfieldMesh/Shaders/Private/VirtualHeightfieldMesh.ush
--- a/Engine/Plugins/Experimental/VirtualHeightfieldMesh/Shaders/Private/VirtualHeightfieldMesh.ush
+++ b/Engine/Plugins/Experimental/VirtualHeightfieldMesh/Shaders/Private/VirtualHeightfieldMesh.ush
@@ -87,37 +87,6 @@ QuadRenderItem UnpackQuadRenderItem(uint2 PackedItem)
 	return Item;
 }
 
-/** Description of neighbor items. We fill 4 of these for each QuadRenderItem. Packs as uint so store in Buffer declared as uint. */
-struct QuadNeighborItem
-{
-	uint PhysicalAddress;
-	bool bLerp;
-};
-
-uint Pack(QuadNeighborItem Item)
-{
-	uint PackedItem = 0;
-	PackedItem |= Item.PhysicalAddress & 0x000FFFFF;
-	PackedItem |= Item.bLerp ? 1 << 20 : 0;
-	return PackedItem;
-}
-
-QuadNeighborItem InitQuadNeighborItem(uint PhysicalAddress, bool bLerp)
-{
-	QuadNeighborItem Item;
-	Item.PhysicalAddress = PhysicalAddress;
-	Item.bLerp = bLerp;
-	return Item;
-}
-
-QuadNeighborItem UnpackQuadNeighborItem(uint PackedItem)
-{
-	QuadNeighborItem Item;
-	Item.PhysicalAddress = PackedItem & 0x000FFFFF;
-	Item.bLerp = ((PackedItem & 0x00100000) >> 20) == 1;
-	return Item;
-}
-
 /** Final render instance description used by the DrawInstancedIndirect(). */
 struct QuadRenderInstance
 {
@@ -134,3 +103,19 @@ uint UnpackLevel(QuadRenderInstance Item)
 {
 	return Item.PosLevelPacked >> 24;
 }
+
+
+/** Shared helper function to calculate a LOD based on distance from camera. */
+float CalculateDistanceLod(float InDistanceSq, float4 InLodFactors)
+{
+	float ScaledDistance = sqrt(InDistanceSq) * InLodFactors.w;
+	float LodForDistance0 = saturate(ScaledDistance / (InLodFactors.x * InLodFactors.y));
+	float LodForDistanceN = log2(1 + max((ScaledDistance / InLodFactors.x - InLodFactors.y), 0)) / log2(InLodFactors.z);
+	return LodForDistance0 + LodForDistanceN;
+}
+
+/** Shared helper function to calculate a LOD bias based on our LodBias texture value. */
+float CalculateBiasLod(float InValue, float InLodBiasScale)
+{
+	return max((InValue - 0.05f) * InLodBiasScale - 1.f, 0);
+}
diff --git a/Engine/Plugins/Experimental/VirtualHeightfieldMesh/Shaders/Private/VirtualHeightfieldMeshVertexFactory.ush b/Engine/Plugins/Experimental/VirtualHeightfieldMesh/Shaders/Private/VirtualHeightfieldMeshVertexFactory.ush
--- a/Engine/Plugins/Experimental/VirtualHeightfieldMesh/Shaders/Private/VirtualHeightfieldMeshVertexFactory.ush
+++ b/Engine/Plugins/Experimental/VirtualHeightfieldMesh/Shaders/Private/VirtualHeightfieldMeshVertexFactory.ush
@@ -20,8 +20,12 @@ float MaxLod;
 float4x4 VirtualHeightfieldToLocal;
 float4x4 VirtualHeightfieldToWorld;
 
+Texture2D<float> LodBiasTexture;
+SamplerState LodBiasSampler;
+
 float3 LodViewOrigin;
 float4 LodDistances;
+float LodBiasScale;
 
 #define GRID_SIZE (VirtualHeightfieldMeshVF.NumQuadsPerTileSide+1)
 
@@ -50,6 +54,12 @@ struct FVertexFactoryInterpolantsVSToPS
 #endif
 };
 
+/** Fetch the lod bias for the location. */
+float GetLodBias(float2 InUV, float InLodBiasScale)
+{
+	return CalculateBiasLod(LodBiasTexture.SampleLevel(LodBiasSampler, InUV, 0), InLodBiasScale);
+}
+
 /** Helper to morph the UV location of a vertex. First snaps down InMorphFactorFloor LOD levels and then morphs InMorphFactorFrac towards the next LOD level. */
 float2 MorphVertex(float2 InLocalUV, uint InGridSize, uint InMorphFactorFloor, float InMorphFactorFrac)
 {
@@ -82,15 +92,6 @@ FVertexFactoryIntermediates GetVertexFactoryIntermediates(FVertexFactoryInput In
 	uint2 VertexCoord = uint2(Input.VertexId % GRID_SIZE, Input.VertexId / GRID_SIZE);
 	float2 LocalUV = (float2)VertexCoord / (float)(GRID_SIZE - 1);
 
-	// Get neighbor details according to location.
-	uint NeighborIndex = ((LocalUV.x + LocalUV.y) > 1) ?
-		(LocalUV.x < LocalUV.y ? 1 : 2) :
-		(LocalUV.x < LocalUV.y ? 0 : 3);
-	float NeighborWeight = ((LocalUV.x + LocalUV.y) > 1) ?
-		(LocalUV.x < LocalUV.y ? LocalUV.y - 0.5 : LocalUV.x - 0.5) :
-		(LocalUV.x < LocalUV.y ? 0.5 - LocalUV.x : 0.5 - LocalUV.y);
-	NeighborWeight *= 2;
-
 	// Calculate vertex UV details before morphing
 	float2 XY = ((float2)Pos + LocalUV) * (float)(1 << Level);
 	float2 NormalizedPos = (XY * PageTableSize.zw);
@@ -102,35 +103,37 @@ FVertexFactoryIntermediates GetVertexFactoryIntermediates(FVertexFactoryInput In
 		float Height = HeightTexture.SampleLevel(HeightSampler, LocalPhysicalUV, 0);
 		float3 WorldPos = mul(float4(NormalizedPos, Height, 1), VirtualHeightfieldToWorld).xyz;
 		float DistanceSq = dot(LodViewOrigin - WorldPos, LodViewOrigin - WorldPos);
-		float Distance = sqrt(DistanceSq) * LodDistances.w;
+		float LodForDistance = CalculateDistanceLod(DistanceSq, LodDistances);
+		float LodBias = GetLodBias(NormalizedPos, LodBiasScale);
 		
-		// Continous LOD at LOD 0.
-		float LodForDistance0 = saturate(Distance / (LodDistances.x * LodDistances.y));
-		// Continous LOD at LOD 1 and above.
-		float LodForDistanceN = log2(1 + max((Distance / LodDistances.x - LodDistances.y), 0)) / log2(LodDistances.z);
-		// Sum to get total continous LOD at distance from view origin.
-		float LodForDistance = LodForDistance0 + LodForDistanceN;
 		// Clamp between the LOD level for this instance and the max LOD. 
 		// Note that the culling phase should already ensure that LOD >= Level.
-		float LodForDistanceClamped = clamp(LodForDistance, (float)Level, MaxLod);
-		// Number of full levels that we need to morph.
-		float LodMorphFloor = floor(LodForDistanceClamped) - (float)Level;
-		// Final fractional morph has smoothstep to push transition region out to end of the LOD level.
-		float LodMorphFrac = smoothstep(0.7f, 0.98f, frac(LodForDistanceClamped));
+		float LodClamped = clamp(LodForDistance - LodBias, (float)Level, MaxLod);
+		
+		// Number of levels that we need to morph.
+		float LodMorphFloor = floor(LodClamped) - (float)Level;
+		float LodMorphCeil = ceil(LodClamped) - (float)Level;
+		float LodMorphFrac = LodClamped - (LodMorphFloor + (float)Level);
 
-		// Use LOD level as sample level but bias to keep sample level in a good range.
-		SampleLevel = max(0, LodForDistanceClamped - 0.5f);
+		// NOTE: Removing fractional continuous LOD here.
+		// This is because fractional locations come away from the surface of the triangles that they interpolate and we see the resultant surface shimmer.
+		// A fix for this while keeping fractional LOD is to use some sort of triangle barycentric interpolation when sampling the height texture instead of bilinear.
+		// But snapping to units here looks OK.
+		LodMorphFrac = 0;
 
 		// Apply morph to vertex postion.
 		LocalUV = MorphVertex(LocalUV, GRID_SIZE - 1, (uint)LodMorphFloor, LodMorphFrac);
 		// Adjust other position variables to match the new morphed position.
 		XY = ((float2)Pos + LocalUV) * (float)(1 << Level);
 		NormalizedPos = (XY * PageTableSize.zw);
+
+		// Use LOD level as sample level but bias to keep sample level in a good range.
+		SampleLevel = max(0, LodClamped - 0.5f);
 	}
 	
 	// Sample height from virtual texture.
 	VTUniform Uniform = VTUniform_Unpack(VTPackedUniform);
-	Uniform.vPageBorderSize -= .5f * PhysicalTextureSize.y; // Half texel offset hack
+	Uniform.vPageBorderSize -= .5f * PhysicalTextureSize.y; // Half texel offset is used in VT write and in sampling because we want texel locations to match landscape vertices.
 	VTPageTableUniform PageTableUniform = VTPageTableUniform_Unpack(VTPackedPageTableUniform0, VTPackedPageTableUniform1);
 	VTPageTableResult VTResult0 = TextureLoadVirtualPageTableLevel(PageTableTexture, PageTableUniform, NormalizedPos, VTADDRESSMODE_CLAMP, VTADDRESSMODE_CLAMP, floor(SampleLevel));
 	float4 Height0 = TextureVirtualSampleLevel(HeightTexture, HeightSampler, VTResult0, 0, Uniform);
diff --git a/Engine/Plugins/Experimental/VirtualHeightfieldMesh/Source/VirtualHeightfieldMesh/Private/VirtualHeightfieldMeshComponent.cpp b/Engine/Plugins/Experimental/VirtualHeightfieldMesh/Source/VirtualHeightfieldMesh/Private/VirtualHeightfieldMeshComponent.cpp
--- a/Engine/Plugins/Experimental/VirtualHeightfieldMesh/Source/VirtualHeightfieldMesh/Private/VirtualHeightfieldMeshComponent.cpp
+++ b/Engine/Plugins/Experimental/VirtualHeightfieldMesh/Source/VirtualHeightfieldMesh/Private/VirtualHeightfieldMeshComponent.cpp
@@ -14,7 +14,10 @@ UVirtualHeightfieldMeshComponent::UVirtualHeightfieldMeshComponent(const FObject
 	: Super(ObjectInitializer)
 {
 	CastShadow = true;
+	bCastContactShadow = false;
 	bUseAsOccluder = true;
+	bAffectDynamicIndirectLighting = false;
+	bAffectDistanceFieldLighting = false;
 	bNeverDistanceCull = true;
 #if WITH_EDITORONLY_DATA
 	bEnableAutoLODGeneration = false;
diff --git a/Engine/Plugins/Experimental/VirtualHeightfieldMesh/Source/VirtualHeightfieldMesh/Private/VirtualHeightfieldMeshSceneProxy.cpp b/Engine/Plugins/Experimental/VirtualHeightfieldMesh/Source/VirtualHeightfieldMesh/Private/VirtualHeightfieldMeshSceneProxy.cpp
--- a/Engine/Plugins/Experimental/VirtualHeightfieldMesh/Source/VirtualHeightfieldMesh/Private/VirtualHeightfieldMeshSceneProxy.cpp
+++ b/Engine/Plugins/Experimental/VirtualHeightfieldMesh/Source/VirtualHeightfieldMesh/Private/VirtualHeightfieldMeshSceneProxy.cpp
@@ -361,7 +361,9 @@ FVirtualHeightfieldMeshSceneProxy::FVirtualHeightfieldMeshSceneProxy(UVirtualHei
 	: FPrimitiveSceneProxy(InComponent, NAME_VirtualHeightfieldMesh)
 	, bHiddenInEditor(InComponent->GetHiddenInEditor())
 	, RuntimeVirtualTexture(InComponent->GetVirtualTexture())
-	, MinMaxTexture(nullptr)
+	, HeightMinMaxTexture(nullptr)
+	, LodBiasTexture(nullptr)
+	, LodBiasMinMaxTexture(nullptr)
 	, AllocatedVirtualTexture(nullptr)
 	, bCallbackRegistered(false)
 	, NumQuadsPerTileSide(0)
@@ -369,6 +371,7 @@ FVirtualHeightfieldMeshSceneProxy::FVirtualHeightfieldMeshSceneProxy(UVirtualHei
 	, Lod0ScreenSize(InComponent->GetLod0ScreenSize())
 	, Lod0Distribution(InComponent->GetLod0Distribution())
 	, LodDistribution(InComponent->GetLodDistribution())
+	, LodBiasScale(InComponent->GetLodBiasScale())
 	, NumSubdivisionLODs(InComponent->GetNumSubdivisionLods())
 	, NumTailLods(InComponent->GetNumTailLods())
 	, NumForceLoadLods(InComponent->GetNumForceLoadLods())
@@ -394,8 +397,11 @@ FVirtualHeightfieldMeshSceneProxy::FVirtualHeightfieldMeshSceneProxy(UVirtualHei
 	UHeightfieldMinMaxTexture* HeightfieldMinMaxTexture = InComponent->GetMinMaxTexture();
 	if (HeightfieldMinMaxTexture != nullptr)
 	{
-		MinMaxTexture = HeightfieldMinMaxTexture->Texture;
+		HeightMinMaxTexture = HeightfieldMinMaxTexture->Texture;
 		BuildOcclusionVolumes(HeightfieldMinMaxTexture->TextureData, HeightfieldMinMaxTexture->TextureDataSize, HeightfieldMinMaxTexture->TextureDataMips, InComponent->GetNumOcclusionLods());
+
+		LodBiasTexture = HeightfieldMinMaxTexture->LodBiasTexture;
+		LodBiasMinMaxTexture = HeightfieldMinMaxTexture->LodBiasMinMaxTexture;
 	}
 }
 
@@ -525,6 +531,7 @@ void FVirtualHeightfieldMeshSceneProxy::GetDynamicMeshElements(const TArray<cons
 				UserData->InstanceBufferSRV = Buffers.InstanceBufferSRV;
 				UserData->PageTableTexture = AllocatedVirtualTexture->GetPageTableTexture(0);
 				UserData->HeightPhysicalTexture = AllocatedVirtualTexture->GetPhysicalTexture(0);
+				UserData->LodBiasTexture = LodBiasTexture ? LodBiasTexture->Resource->TextureRHI : GBlackTexture->TextureRHI;
 
 				AllocatedVirtualTexture->GetPackedUniform(&UserData->PackedUniform, 0);
 				AllocatedVirtualTexture->GetPackedPageTableUniform(UserData->PackedPageTableUniform);
@@ -554,6 +561,7 @@ void FVirtualHeightfieldMeshSceneProxy::GetDynamicMeshElements(const TArray<cons
 #endif
 				
 				UserData->LodDistances = VirtualHeightfieldMesh::CalculateLodRanges(MainView, this);
+				UserData->LodBiasScale = LodBiasScale;
 				
 				BatchElement.PrimitiveIdMode = PrimID_ForceZero;
 				BatchElement.PrimitiveUniformBuffer = GetUniformBuffer();
@@ -720,14 +728,18 @@ namespace VirtualHeightfieldMesh
 		SHADER_USE_PARAMETER_STRUCT(FCollectQuadsCS, FGlobalShader);
 
 		BEGIN_SHADER_PARAMETER_STRUCT(FParameters, )
-			SHADER_PARAMETER_TEXTURE(Texture2D, MinMaxTexture)
+			SHADER_PARAMETER_TEXTURE(Texture2D, HeightMinMaxTexture)
 			SHADER_PARAMETER_SAMPLER(SamplerState, MinMaxTextureSampler)
 			SHADER_PARAMETER(int32, MinMaxLevelOffset)
+			SHADER_PARAMETER_TEXTURE(Texture2D, LodBiasMinMaxTexture)
 			SHADER_PARAMETER_TEXTURE(Texture2D<float>, OcclusionTexture)
 			SHADER_PARAMETER(int32, OcclusionLevelOffset)
 			SHADER_PARAMETER_TEXTURE(Texture2D<uint>, PageTableTexture)
+			SHADER_PARAMETER(uint32, MaxLevel)
 			SHADER_PARAMETER(FVector4, PageTableSize)
+			SHADER_PARAMETER(uint32, PageTableFeedbackId)
 			SHADER_PARAMETER(FVector4, LodDistances)
+			SHADER_PARAMETER(float, LodBiasScale)
 			SHADER_PARAMETER(FVector, ViewOrigin)
 			SHADER_PARAMETER_ARRAY(FVector4, FrustumPlanes, [5])
 			SHADER_PARAMETER(FMatrix, UVToWorld)
@@ -737,84 +749,7 @@ namespace VirtualHeightfieldMesh
 			SHADER_PARAMETER_RDG_BUFFER_UAV(RWBuffer<uint>, RWQueueBuffer)
 			SHADER_PARAMETER_RDG_BUFFER_UAV(RWBuffer<uint2>, RWQuadBuffer)
 			SHADER_PARAMETER_RDG_BUFFER_UAV(RWBuffer<uint>, RWIndirectArgsBuffer)
-		END_SHADER_PARAMETER_STRUCT()
-
-		static bool ShouldCompilePermutation(FGlobalShaderPermutationParameters const& Parameters)
-		{
-			return IsFeatureLevelSupported(Parameters.Platform, ERHIFeatureLevel::SM5);
-		}
-	};
-
-	IMPLEMENT_GLOBAL_SHADER(FCollectQuadsCS, "/Plugin/VirtualHeightfieldMesh/Private/VirtualHeightfieldMesh.usf", "CollectQuadsCS", SF_Compute);
-
-	/** Shader that draws to a render target the Lod info for the quads output by the Collect pass. */
-	class FRenderLodMap : public FGlobalShader
-	{
-	public:
-		SHADER_USE_PARAMETER_STRUCT(FRenderLodMap, FGlobalShader);
-
-		BEGIN_SHADER_PARAMETER_STRUCT(FParameters, )
-			RENDER_TARGET_BINDING_SLOTS()
-			SHADER_PARAMETER(FVector4, PageTableSize)
-			SHADER_PARAMETER_RDG_BUFFER_SRV(Buffer<uint2>, QuadBuffer)
-			RDG_BUFFER_ACCESS(IndirectArgsBuffer, ERHIAccess::IndirectArgs)
-		END_SHADER_PARAMETER_STRUCT()
-
-		static bool ShouldCompilePermutation(FGlobalShaderPermutationParameters const& Parameters)
-		{
-			return IsFeatureLevelSupported(Parameters.Platform, ERHIFeatureLevel::SM5);
-		}
-	};
-
-	/** RenderLodMap vertex shader. */
-	class FRenderLodMapVS : public FRenderLodMap
-	{
-	public:
-		DECLARE_GLOBAL_SHADER(FRenderLodMapVS);
-
-		FRenderLodMapVS()
-		{}
-
-		FRenderLodMapVS(const ShaderMetaType::CompiledShaderInitializerType& Initializer)
-			: FRenderLodMap(Initializer)
-		{}
-	};
-
-	IMPLEMENT_GLOBAL_SHADER(FRenderLodMapVS, "/Plugin/VirtualHeightfieldMesh/Private/VirtualHeightfieldMesh.usf", "RenderLodMapVS", SF_Vertex);
-
-	/** RenderLodMap pixel shader. */
-	class FRenderLodMapPS : public FRenderLodMap
-	{
-	public:
-		DECLARE_GLOBAL_SHADER(FRenderLodMapPS);
-
-		FRenderLodMapPS()
-		{}
-
-		FRenderLodMapPS(const ShaderMetaType::CompiledShaderInitializerType& Initializer)
-			: FRenderLodMap(Initializer)
-		{}
-	};
-
-	IMPLEMENT_GLOBAL_SHADER(FRenderLodMapPS, "/Plugin/VirtualHeightfieldMesh/Private/VirtualHeightfieldMesh.usf", "RenderLodMapPS", SF_Pixel);
-
-	/** ResolveNeighborLods compute shader. */
-	class FResolveNeighborLodsCS : public FGlobalShader
-	{
-	public:
-		DECLARE_GLOBAL_SHADER(FResolveNeighborLodsCS);
-		SHADER_USE_PARAMETER_STRUCT(FResolveNeighborLodsCS, FGlobalShader);
-
-		BEGIN_SHADER_PARAMETER_STRUCT(FParameters, )
-			SHADER_PARAMETER(FVector4, PageTableSize)
-			SHADER_PARAMETER_TEXTURE(Texture2D, PageTableTexture)
-			SHADER_PARAMETER(uint32, PageTableFeedbackId)
-			SHADER_PARAMETER_RDG_BUFFER_SRV(Buffer<uint2>, QuadBuffer)
-			SHADER_PARAMETER_RDG_TEXTURE(Texture2D<float2>, LodTexture)
-			SHADER_PARAMETER_RDG_BUFFER_SRV(Buffer<uint>, IndirectArgsBufferSRV)
-			SHADER_PARAMETER_RDG_BUFFER_UAV(RWBuffer<uint>, RWQuadNeighborBuffer)
 			SHADER_PARAMETER_RDG_BUFFER_UAV(RWBuffer<uint>, RWFeedbackBuffer)
-			RDG_BUFFER_ACCESS(IndirectArgsBuffer, ERHIAccess::IndirectArgs)
 		END_SHADER_PARAMETER_STRUCT()
 
 		static bool ShouldCompilePermutation(FGlobalShaderPermutationParameters const& Parameters)
@@ -823,7 +758,7 @@ namespace VirtualHeightfieldMesh
 		}
 	};
 
-	IMPLEMENT_GLOBAL_SHADER(FResolveNeighborLodsCS, "/Plugin/VirtualHeightfieldMesh/Private/VirtualHeightfieldMesh.usf", "ResolveNeighborLodsCS", SF_Compute);
+	IMPLEMENT_GLOBAL_SHADER(FCollectQuadsCS, "/Plugin/VirtualHeightfieldMesh/Private/VirtualHeightfieldMesh.usf", "CollectQuadsCS", SF_Compute);
 
 	/** InitInstanceBuffer compute shader. */
 	class FInitInstanceBufferCS : public FGlobalShader
@@ -852,15 +787,15 @@ namespace VirtualHeightfieldMesh
 		SHADER_USE_PARAMETER_STRUCT(FCullInstances, FGlobalShader);
 
 		BEGIN_SHADER_PARAMETER_STRUCT(FParameters, )
-			SHADER_PARAMETER_TEXTURE(Texture2D, MinMaxTexture)
+			SHADER_PARAMETER_TEXTURE(Texture2D, HeightMinMaxTexture)
 			SHADER_PARAMETER_SAMPLER(SamplerState, MinMaxTextureSampler)
+			SHADER_PARAMETER(int32, MinMaxLevelOffset)
 			SHADER_PARAMETER_TEXTURE(Texture2D, PageTableTexture)
 			SHADER_PARAMETER(FVector4, PageTableSize)
 			SHADER_PARAMETER_ARRAY(FVector4, FrustumPlanes, [5])
 			SHADER_PARAMETER(FVector4, PhysicalPageTransform)
 			SHADER_PARAMETER(uint32, NumPhysicalAddressBits)
 			SHADER_PARAMETER_RDG_BUFFER_SRV(Buffer<uint2>, QuadBuffer)
-			SHADER_PARAMETER_RDG_BUFFER_SRV(Buffer<uint>, QuadNeighborBuffer)
 			SHADER_PARAMETER_RDG_BUFFER_SRV(Buffer<uint>, IndirectArgsBufferSRV)
 			SHADER_PARAMETER_UAV(RWStructuredBuffer<QuadRenderInstance>, RWInstanceBuffer)
 			SHADER_PARAMETER_UAV(RWBuffer<uint>, RWIndirectArgsBuffer)
@@ -899,7 +834,7 @@ namespace VirtualHeightfieldMesh
 
 
 	/** Default Min/Max texture has the fixed maximum [0,1]. */
-	class FMinMaxDefaultTexture : public FTexture
+	class FHeightMinMaxDefaultTexture : public FTexture
 	{
 	public:
 		virtual void InitRHI() override
@@ -924,7 +859,7 @@ namespace VirtualHeightfieldMesh
 	};
 
 	/** Single global instance of default Min/Max texture. */
-	FTexture* GMinMaxDefaultTexture = new TGlobalResource<FMinMaxDefaultTexture>;
+	FTexture* GHeightMinMaxDefaultTexture = new TGlobalResource<FHeightMinMaxDefaultTexture>;
 
 	/** View matrices that can be frozen in freezerendering mode. */
 	struct FViewData
@@ -987,7 +922,8 @@ namespace VirtualHeightfieldMesh
 	struct FProxyDesc
 	{
 		FRHITexture* PageTableTexture;
-		FRHITexture* MinMaxTexture;
+		FRHITexture* HeightMinMaxTexture;
+		FRHITexture* LodBiasMinMaxTexture;
 		int32 MinMaxLevelOffset;
 
 		uint32 MaxLevel;
@@ -1012,6 +948,7 @@ namespace VirtualHeightfieldMesh
 		FSceneView const* ViewDebug;
 		FVector ViewOrigin;
 		FVector4 LodDistances;
+		float LodBiasScale;
 		FVector4 Planes[5];
 		FTextureRHIRef OcclusionTexture;
 		int32 OcclusionLevelOffset;
@@ -1043,12 +980,6 @@ namespace VirtualHeightfieldMesh
 		FRDGBufferRef IndirectArgsBuffer;
 		FRDGBufferUAVRef IndirectArgsBufferUAV;
 		FRDGBufferSRVRef IndirectArgsBufferSRV;
-
-		FRDGTextureRef LodTexture;
-
-		FRDGBufferRef QuadNeighborBuffer;
-		FRDGBufferUAVRef QuadNeighborBufferUAV;
-		FRDGBufferSRVRef QuadNeighborBufferSRV;
 	};
 
 	/** Initialize the FDrawInstanceBuffers objects. */
@@ -1091,17 +1022,6 @@ namespace VirtualHeightfieldMesh
 		OutResources.IndirectArgsBuffer = GraphBuilder.CreateBuffer(FRDGBufferDesc::CreateIndirectDesc(IndirectArgsByteSize), TEXT("IndirectArgsBuffer"));
 		OutResources.IndirectArgsBufferUAV = GraphBuilder.CreateUAV(OutResources.IndirectArgsBuffer);
 		OutResources.IndirectArgsBufferSRV = GraphBuilder.CreateSRV(OutResources.IndirectArgsBuffer);
-
-		FRDGTextureDesc LodTextureDesc = FRDGTextureDesc::Create2D(
-			FIntPoint(InDesc.PageTableSize.X, InDesc.PageTableSize.Y),
-			PF_R8G8,
-			FClearValueBinding::None,
-			TexCreate_RenderTargetable | TexCreate_ShaderResource);
-		OutResources.LodTexture = GraphBuilder.CreateTexture(LodTextureDesc, TEXT("LodTexture"));
-
-		OutResources.QuadNeighborBuffer = GraphBuilder.CreateBuffer(FRDGBufferDesc::CreateBufferDesc(sizeof(uint32), InDesc.MaxRenderItems * 4), TEXT("QuadNeighborBuffer"));
-		OutResources.QuadNeighborBufferUAV = GraphBuilder.CreateUAV(FRDGBufferUAVDesc(OutResources.QuadNeighborBuffer, PF_R32_UINT));
-		OutResources.QuadNeighborBufferSRV = GraphBuilder.CreateSRV(FRDGBufferSRVDesc(OutResources.QuadNeighborBuffer, PF_R32_UINT));
 	}
 
 	/** Transition our output draw buffers for use. Read or write access is set according to the bToWrite parameter. */
@@ -1174,17 +1094,21 @@ namespace VirtualHeightfieldMesh
 		TShaderMapRef<FCollectQuadsCS> ComputeShader(InGlobalShaderMap);
 
 		FCollectQuadsCS::FParameters* PassParameters = GraphBuilder.AllocParameters<FCollectQuadsCS::FParameters>();
-		PassParameters->MinMaxTexture = InDesc.MinMaxTexture;
+		PassParameters->HeightMinMaxTexture = InDesc.HeightMinMaxTexture;
+		PassParameters->LodBiasMinMaxTexture = InDesc.LodBiasMinMaxTexture;
 		PassParameters->MinMaxTextureSampler = TStaticSamplerState<SF_Point>::GetRHI();
 		PassParameters->MinMaxLevelOffset = InDesc.MinMaxLevelOffset;
 		PassParameters->OcclusionTexture = InViewDesc.OcclusionTexture;
 		PassParameters->OcclusionLevelOffset = InViewDesc.OcclusionLevelOffset;
 		PassParameters->PageTableTexture = InDesc.PageTableTexture;
+		PassParameters->MaxLevel = InDesc.MaxLevel;
 		PassParameters->PageTableSize = InDesc.PageTableSize;
+		PassParameters->PageTableFeedbackId = InDesc.PageTableFeedbackId;
 		PassParameters->UVToWorld = InDesc.UVToWorld;
 		PassParameters->UVToWorldScale = InDesc.UVToWorldScale;
 		PassParameters->ViewOrigin = InViewDesc.ViewOrigin;
 		PassParameters->LodDistances = InViewDesc.LodDistances;
+		PassParameters->LodBiasScale = InViewDesc.LodBiasScale;
 		for (int32 PlaneIndex = 0; PlaneIndex < 5; ++PlaneIndex)
 		{
 			PassParameters->FrustumPlanes[PlaneIndex] = InViewDesc.Planes[PlaneIndex];
@@ -1194,6 +1118,7 @@ namespace VirtualHeightfieldMesh
 		PassParameters->RWQueueBuffer = InVolatileResources.QueueBufferUAV;
 		PassParameters->RWQuadBuffer = InVolatileResources.QuadBufferUAV;
 		PassParameters->RWIndirectArgsBuffer = InVolatileResources.IndirectArgsBufferUAV;
+		PassParameters->RWFeedbackBuffer = InVolatileResources.FeedbackBufferUAV;
 
 		FComputeShaderUtils::AddPass(
 			GraphBuilder,
@@ -1201,72 +1126,6 @@ namespace VirtualHeightfieldMesh
 			ComputeShader, PassParameters, FIntVector(InDesc.NumCollectPassWavefronts, 1, 1));
 	}
 
-	/** Render Lod info map, which can then be sampled to determine neighbor Lod transitions. */
-	void AddPass_RenderLodMap(FRDGBuilder& GraphBuilder, FGlobalShaderMap* InGlobalShaderMap, FProxyDesc const& InDesc, FVolatileResources& InVolatileResources)
-	{
-		TShaderMapRef< FRenderLodMapVS > VertexShader(InGlobalShaderMap);
-		TShaderMapRef< FRenderLodMapPS > PixelShader(InGlobalShaderMap);
-
-		FRenderLodMap::FParameters* PassParameters = GraphBuilder.AllocParameters<FRenderLodMap::FParameters>();
-		PassParameters->RenderTargets[0] = FRenderTargetBinding(InVolatileResources.LodTexture, ERenderTargetLoadAction::ENoAction);
-		PassParameters->QuadBuffer = InVolatileResources.QuadBufferSRV;
-		PassParameters->PageTableSize = InDesc.PageTableSize;
-		PassParameters->IndirectArgsBuffer = InVolatileResources.IndirectArgsBuffer;
-
-		GraphBuilder.AddPass(
-			RDG_EVENT_NAME("RenderLodMap"),
-			PassParameters,
-			ERDGPassFlags::Raster,
-			[VertexShader, PixelShader, PassParameters, IndirectArgsBuffer = InVolatileResources.IndirectArgsBuffer](FRHICommandListImmediate& RHICmdListImmediate)
-		{
-				FGraphicsPipelineStateInitializer GraphicsPSOInit;
-				RHICmdListImmediate.ApplyCachedRenderTargets(GraphicsPSOInit);
-				GraphicsPSOInit.DepthStencilState = TStaticDepthStencilState<false, CF_Always>::GetRHI();
-				GraphicsPSOInit.BlendState = TStaticBlendState<CW_RGBA, BO_Add, BF_One, BF_InverseSourceAlpha, BO_Add, BF_Zero, BF_One>::GetRHI();
-				GraphicsPSOInit.RasterizerState = TStaticRasterizerState<>::GetRHI();
-				GraphicsPSOInit.PrimitiveType = PT_TriangleList;
-				GraphicsPSOInit.BoundShaderState.VertexDeclarationRHI = GetVertexDeclarationFVector4();
-				GraphicsPSOInit.BoundShaderState.VertexShaderRHI = VertexShader.GetVertexShader();
-				GraphicsPSOInit.BoundShaderState.PixelShaderRHI = PixelShader.GetPixelShader();
-				SetGraphicsPipelineState(RHICmdListImmediate, GraphicsPSOInit);
-
-				SetShaderParameters(RHICmdListImmediate, VertexShader, VertexShader.GetVertexShader(), *PassParameters);
-				SetShaderParameters(RHICmdListImmediate, PixelShader, PixelShader.GetPixelShader(), *PassParameters);
-
-				int32 IndirectArgOffset = VirtualHeightfieldMesh::IndirectArgsByteOffset_RenderLodMap;
-				RHICmdListImmediate.DrawIndexedPrimitiveIndirect(GTwoTrianglesIndexBuffer.IndexBufferRHI, IndirectArgsBuffer->GetIndirectRHICallBuffer(), IndirectArgOffset);
-				IndirectArgsBuffer->MarkResourceAsUsed();
-		});
-	}
-
-	/** Resolve Neighbor Lod transitions from the Lod map. */
-	void AddPass_ResolveNeighborLods(FRDGBuilder& GraphBuilder, FGlobalShaderMap* InGlobalShaderMap, FProxyDesc const& InDesc, FVolatileResources& InVolatileResources, FMainViewDesc const& InViewDesc)
-	{
-		TShaderMapRef<FResolveNeighborLodsCS> ComputeShader(InGlobalShaderMap);
-
-		FResolveNeighborLodsCS::FParameters* PassParameters = GraphBuilder.AllocParameters<FResolveNeighborLodsCS::FParameters>();
-		PassParameters->PageTableSize = InDesc.PageTableSize;
-		PassParameters->PageTableFeedbackId = InDesc.PageTableFeedbackId;
-		PassParameters->PageTableTexture = InDesc.PageTableTexture;
-		PassParameters->QuadBuffer = InVolatileResources.QuadBufferSRV;
-		PassParameters->LodTexture = InVolatileResources.LodTexture;
-		PassParameters->IndirectArgsBuffer = InVolatileResources.IndirectArgsBuffer;
-		PassParameters->IndirectArgsBufferSRV = InVolatileResources.IndirectArgsBufferSRV;
-		PassParameters->RWQuadNeighborBuffer = InVolatileResources.QuadNeighborBufferUAV;
-		PassParameters->RWFeedbackBuffer = InVolatileResources.FeedbackBufferUAV;
-
-		GraphBuilder.AddPass(
-			RDG_EVENT_NAME("ResolveNeighborLods"),
-			PassParameters,
-			ERDGPassFlags::Compute,
-			[PassParameters, ComputeShader, IndirectArgsBuffer = InVolatileResources.IndirectArgsBuffer](FRHICommandList& RHICmdList)
-		{
-				int32 IndirectArgOffset = VirtualHeightfieldMesh::IndirectArgsByteOffset_FetchNeighborLod;
-				FComputeShaderUtils::DispatchIndirect(RHICmdList, ComputeShader, *PassParameters, IndirectArgsBuffer->GetIndirectRHICallBuffer(), IndirectArgOffset);
-				IndirectArgsBuffer->MarkResourceAsUsed();
-		});
-	}
-
 	/** Initialise the draw indirect buffer. */
 	void AddPass_InitInstanceBuffer(FRDGBuilder& GraphBuilder, FGlobalShaderMap* InGlobalShaderMap, int32 NumQuadsPerTileSide, FDrawInstanceBuffers& InOutputResources)
 	{
@@ -1286,8 +1145,9 @@ namespace VirtualHeightfieldMesh
 	void AddPass_CullInstances(FRDGBuilder& GraphBuilder, FGlobalShaderMap* InGlobalShaderMap, FProxyDesc const& InDesc, FVolatileResources& InVolatileResources, FDrawInstanceBuffers& InOutputResources, FChildViewDesc const& InViewDesc)
 	{
 		FCullInstances::FParameters* PassParameters = GraphBuilder.AllocParameters<FCullInstances::FParameters>();
-		PassParameters->MinMaxTexture = InDesc.MinMaxTexture;
+		PassParameters->HeightMinMaxTexture = InDesc.HeightMinMaxTexture;
 		PassParameters->MinMaxTextureSampler = TStaticSamplerState<SF_Point>::GetRHI();
+		PassParameters->MinMaxLevelOffset = InDesc.MinMaxLevelOffset;
 		PassParameters->PageTableTexture = InDesc.PageTableTexture;
 		PassParameters->PageTableSize = InDesc.PageTableSize;
 		for (int32 PlaneIndex = 0; PlaneIndex < 5; ++PlaneIndex)
@@ -1297,7 +1157,6 @@ namespace VirtualHeightfieldMesh
 		PassParameters->PhysicalPageTransform = InDesc.PhysicalPageTransform;
 		PassParameters->NumPhysicalAddressBits = InDesc.NumPhysicalAddressBits;
 		PassParameters->QuadBuffer = InVolatileResources.QuadBufferSRV;
-		PassParameters->QuadNeighborBuffer = InVolatileResources.QuadNeighborBufferSRV;
 		PassParameters->IndirectArgsBuffer = InVolatileResources.IndirectArgsBuffer;
 		PassParameters->IndirectArgsBufferSRV = InVolatileResources.IndirectArgsBufferSRV;
 		PassParameters->RWInstanceBuffer = InOutputResources.InstanceBufferUAV;
@@ -1324,25 +1183,6 @@ namespace VirtualHeightfieldMesh
 			IndirectArgsBuffer->MarkResourceAsUsed();
 		});
 	}
-
-	/** Add passes for setting up culling for a main view. */
-	void GPUCullMainView( FRDGBuilder& GraphBuilder, FGlobalShaderMap* InGlobalShaderMap, FProxyDesc const& InDesc, FVolatileResources& InVolatileResources, FMainViewDesc const& InViewDesc)
-	{
-		RDG_EVENT_SCOPE(GraphBuilder, "MainView");
-
-		AddPass_InitBuffers(GraphBuilder, InGlobalShaderMap, InDesc, InVolatileResources);
-		AddPass_CollectQuads(GraphBuilder, InGlobalShaderMap, InDesc, InVolatileResources, InViewDesc);
-		AddPass_RenderLodMap(GraphBuilder, InGlobalShaderMap, InDesc, InVolatileResources);
-		AddPass_ResolveNeighborLods(GraphBuilder, InGlobalShaderMap, InDesc, InVolatileResources, InViewDesc);
-	}
-
-	/** Add passes for culling a single child view. */
-	void GPUCullChildView(FRDGBuilder& GraphBuilder, FGlobalShaderMap* InGlobalShaderMap, FProxyDesc const& InDesc, FVolatileResources& InVolatileResources, FDrawInstanceBuffers& InOutputResources, FChildViewDesc const& InViewDesc)
-	{
-		RDG_EVENT_SCOPE(GraphBuilder, "CullView");
-
-		AddPass_CullInstances(GraphBuilder, InGlobalShaderMap, InDesc, InVolatileResources, InOutputResources, InViewDesc);
-	}
 }
 
 void FVirtualHeightfieldMeshRendererExtension::SubmitWork(FRDGBuilder& GraphBuilder)
@@ -1392,8 +1232,9 @@ void FVirtualHeightfieldMeshRendererExtension::SubmitWork(FRDGBuilder& GraphBuil
 
 		VirtualHeightfieldMesh::FProxyDesc ProxyDesc;
 		ProxyDesc.PageTableTexture = AllocatedVirtualTexture->GetPageTableTexture(0);
-		ProxyDesc.MinMaxTexture = Proxy->MinMaxTexture ? Proxy->MinMaxTexture->Resource->TextureRHI : VirtualHeightfieldMesh::GMinMaxDefaultTexture->TextureRHI;
-		ProxyDesc.MinMaxLevelOffset = ProxyDesc.MinMaxTexture->GetNumMips() - 1 - AllocatedVirtualTexture->GetMaxLevel();
+		ProxyDesc.HeightMinMaxTexture = Proxy->HeightMinMaxTexture ? Proxy->HeightMinMaxTexture->Resource->TextureRHI : VirtualHeightfieldMesh::GHeightMinMaxDefaultTexture->TextureRHI;
+		ProxyDesc.LodBiasMinMaxTexture = Proxy->LodBiasMinMaxTexture ? Proxy->LodBiasMinMaxTexture->Resource->TextureRHI : GBlackTexture->TextureRHI;
+		ProxyDesc.MinMaxLevelOffset = ProxyDesc.HeightMinMaxTexture->GetNumMips() - 1 - AllocatedVirtualTexture->GetMaxLevel();
 		ProxyDesc.MaxLevel = AllocatedVirtualTexture->GetMaxLevel();
 		ProxyDesc.NumForceLoadLods = FMath::Min<uint32>(FMath::Max(Proxy->NumForceLoadLods, 0), ProxyDesc.MaxLevel + 1);
 		ProxyDesc.PageTableSize = PageTableSize;
@@ -1422,6 +1263,7 @@ void FVirtualHeightfieldMeshRendererExtension::SubmitWork(FRDGBuilder& GraphBuil
 			// ViewOrigin and Frustum Planes are all converted to UV space for the shader.
 			MainViewDesc.ViewOrigin = Proxy->WorldToUV.TransformPosition(MainViewData.ViewOrigin);
 			MainViewDesc.LodDistances = VirtualHeightfieldMesh::CalculateLodRanges(MainView, Proxy);
+			MainViewDesc.LodBiasScale = Proxy->LodBiasScale;
 
 			const int32 MainViewNumPlanes = FMath::Min(MainViewData.ViewFrustum.Planes.Num(), 5);
 			for (int32 PlaneIndex = 0; PlaneIndex < MainViewNumPlanes; ++PlaneIndex)
@@ -1452,7 +1294,8 @@ void FVirtualHeightfieldMeshRendererExtension::SubmitWork(FRDGBuilder& GraphBuil
 			VirtualHeightfieldMesh::InitializeResources(GraphBuilder, ProxyDesc, MainViewDesc, VolatileResources);
 
 			// Build graph
-			VirtualHeightfieldMesh::GPUCullMainView(GraphBuilder, GetGlobalShaderMap(GMaxRHIFeatureLevel), ProxyDesc, VolatileResources, MainViewDesc);
+			VirtualHeightfieldMesh::AddPass_InitBuffers(GraphBuilder, GetGlobalShaderMap(GMaxRHIFeatureLevel), ProxyDesc, VolatileResources);
+			VirtualHeightfieldMesh::AddPass_CollectQuads(GraphBuilder, GetGlobalShaderMap(GMaxRHIFeatureLevel), ProxyDesc, VolatileResources, MainViewDesc);
 
 			FVirtualTextureFeedbackBufferDesc Desc;
  			Desc.Init(CVarVHMMaxFeedbackItems.GetValueOnRenderThread() + 1);
@@ -1484,7 +1327,7 @@ void FVirtualHeightfieldMeshRendererExtension::SubmitWork(FRDGBuilder& GraphBuil
 				}
 
 				// Build graph
-				VirtualHeightfieldMesh::GPUCullChildView(GraphBuilder, GetGlobalShaderMap(GMaxRHIFeatureLevel), ProxyDesc, VolatileResources, Buffers[WorkDescs[WorkIndex].BufferIndex], ChildViewDesc);
+				VirtualHeightfieldMesh::AddPass_CullInstances(GraphBuilder, GetGlobalShaderMap(GMaxRHIFeatureLevel), ProxyDesc, VolatileResources, Buffers[WorkDescs[WorkIndex].BufferIndex], ChildViewDesc);
 
 				WorkIndex++;
 			}
diff --git a/Engine/Plugins/Experimental/VirtualHeightfieldMesh/Source/VirtualHeightfieldMesh/Private/VirtualHeightfieldMeshSceneProxy.h b/Engine/Plugins/Experimental/VirtualHeightfieldMesh/Source/VirtualHeightfieldMesh/Private/VirtualHeightfieldMeshSceneProxy.h
--- a/Engine/Plugins/Experimental/VirtualHeightfieldMesh/Source/VirtualHeightfieldMesh/Private/VirtualHeightfieldMeshSceneProxy.h
+++ b/Engine/Plugins/Experimental/VirtualHeightfieldMesh/Source/VirtualHeightfieldMesh/Private/VirtualHeightfieldMeshSceneProxy.h
@@ -34,7 +34,10 @@ public:
 
 	class URuntimeVirtualTexture* RuntimeVirtualTexture;
 	class FMaterialRenderProxy* Material;
-	class UTexture2D* MinMaxTexture;
+	
+	class UTexture2D* HeightMinMaxTexture;
+	class UTexture2D* LodBiasTexture;
+	class UTexture2D* LodBiasMinMaxTexture;
 
 	class IAllocatedVirtualTexture* AllocatedVirtualTexture;
 	bool bCallbackRegistered;
@@ -52,6 +55,7 @@ public:
 	float Lod0ScreenSize;
 	float Lod0Distribution;
 	float LodDistribution;
+	float LodBiasScale;
 	
 	int32 NumSubdivisionLODs;
 	int32 NumTailLods;
diff --git a/Engine/Plugins/Experimental/VirtualHeightfieldMesh/Source/VirtualHeightfieldMesh/Private/VirtualHeightfieldMeshVertexFactory.cpp b/Engine/Plugins/Experimental/VirtualHeightfieldMesh/Source/VirtualHeightfieldMesh/Private/VirtualHeightfieldMeshVertexFactory.cpp
--- a/Engine/Plugins/Experimental/VirtualHeightfieldMesh/Source/VirtualHeightfieldMesh/Private/VirtualHeightfieldMeshVertexFactory.cpp
+++ b/Engine/Plugins/Experimental/VirtualHeightfieldMesh/Source/VirtualHeightfieldMesh/Private/VirtualHeightfieldMeshVertexFactory.cpp
@@ -103,6 +103,8 @@ public:
 		PageTableTextureParameter.Bind(ParameterMap, TEXT("PageTableTexture"));
 		HeightTextureParameter.Bind(ParameterMap, TEXT("HeightTexture"));
 		HeightSamplerParameter.Bind(ParameterMap, TEXT("HeightSampler"));
+		LodBiasTextureParameter.Bind(ParameterMap, TEXT("LodBiasTexture"));
+		LodBiasSamplerParameter.Bind(ParameterMap, TEXT("LodBiasSampler"));
 		PackedUniformParameter.Bind(ParameterMap, TEXT("VTPackedUniform"));
 		PackedPageTableUniform0Parameter.Bind(ParameterMap, TEXT("VTPackedPageTableUniform0"));
 		PackedPageTableUniform1Parameter.Bind(ParameterMap, TEXT("VTPackedPageTableUniform1"));
@@ -113,6 +115,7 @@ public:
 		VirtualHeightfieldToWorldParameter.Bind(ParameterMap, TEXT("VirtualHeightfieldToWorld"));
 		LodViewOriginParameter.Bind(ParameterMap, TEXT("LodViewOrigin"));
 		LodDistancesParameter.Bind(ParameterMap, TEXT("LodDistances"));
+		LodBiasScaleParameter.Bind(ParameterMap, TEXT("LodBiasScale"));
 	}
 
 	void GetElementShaderBindings(
@@ -145,6 +148,10 @@ public:
 			{
 				ShaderBindings.AddTexture(HeightTextureParameter, HeightSamplerParameter, TStaticSamplerState<SF_Bilinear>::GetRHI(), UserData->HeightPhysicalTexture);
 			}
+			if (LodBiasTextureParameter.IsBound() && LodBiasSamplerParameter.IsBound())
+			{
+				ShaderBindings.AddTexture(LodBiasTextureParameter, LodBiasSamplerParameter, TStaticSamplerState<SF_Point>::GetRHI(), UserData->LodBiasTexture);
+			}
 			if (PackedUniformParameter.IsBound())
 			{
 				ShaderBindings.Add(PackedUniformParameter, UserData->PackedUniform);
@@ -185,6 +192,10 @@ public:
 			{
 				ShaderBindings.Add(LodDistancesParameter, UserData->LodDistances);
 			}
+			if (LodBiasScaleParameter.IsBound())
+			{
+				ShaderBindings.Add(LodBiasScaleParameter, UserData->LodBiasScale);
+			}
 		}
 	}
 
@@ -194,6 +205,8 @@ protected:
 	LAYOUT_FIELD(FShaderResourceParameter, PageTableSamplerParameter);
 	LAYOUT_FIELD(FShaderResourceParameter, HeightTextureParameter);
 	LAYOUT_FIELD(FShaderResourceParameter, HeightSamplerParameter);
+	LAYOUT_FIELD(FShaderResourceParameter, LodBiasTextureParameter);
+	LAYOUT_FIELD(FShaderResourceParameter, LodBiasSamplerParameter);
 	LAYOUT_FIELD(FShaderParameter, PackedUniformParameter);
 	LAYOUT_FIELD(FShaderParameter, PackedPageTableUniform0Parameter);
 	LAYOUT_FIELD(FShaderParameter, PackedPageTableUniform1Parameter);
@@ -204,6 +217,7 @@ protected:
 	LAYOUT_FIELD(FShaderParameter, VirtualHeightfieldToWorldParameter);
 	LAYOUT_FIELD(FShaderParameter, LodViewOriginParameter);
 	LAYOUT_FIELD(FShaderParameter, LodDistancesParameter);
+	LAYOUT_FIELD(FShaderParameter, LodBiasScaleParameter);
 };
 
 IMPLEMENT_TYPE_LAYOUT(FVirtualHeightfieldMeshVertexFactoryShaderParameters);
diff --git a/Engine/Plugins/Experimental/VirtualHeightfieldMesh/Source/VirtualHeightfieldMesh/Private/VirtualHeightfieldMeshVertexFactory.h b/Engine/Plugins/Experimental/VirtualHeightfieldMesh/Source/VirtualHeightfieldMesh/Private/VirtualHeightfieldMeshVertexFactory.h
--- a/Engine/Plugins/Experimental/VirtualHeightfieldMesh/Source/VirtualHeightfieldMesh/Private/VirtualHeightfieldMeshVertexFactory.h
+++ b/Engine/Plugins/Experimental/VirtualHeightfieldMesh/Source/VirtualHeightfieldMesh/Private/VirtualHeightfieldMeshVertexFactory.h
@@ -25,6 +25,7 @@ struct FVirtualHeightfieldMeshUserData : public FOneFrameResource
 	FRHIShaderResourceView* InstanceBufferSRV;
 	FRHITexture* PageTableTexture;
 	FRHITexture* HeightPhysicalTexture;
+	FRHITexture* LodBiasTexture;
 	FUintVector4 PackedUniform;
 	FUintVector4 PackedPageTableUniform[2];
 	FVector4 PageTableSize;
@@ -34,6 +35,7 @@ struct FVirtualHeightfieldMeshUserData : public FOneFrameResource
 	FMatrix VirtualHeightfieldToWorld;
 	FVector LodViewOrigin;
 	FVector4 LodDistances;
+	float LodBiasScale;
 };
 
 /**
diff --git a/Engine/Plugins/Experimental/VirtualHeightfieldMesh/Source/VirtualHeightfieldMesh/Public/VirtualHeightfieldMeshComponent.h b/Engine/Plugins/Experimental/VirtualHeightfieldMesh/Source/VirtualHeightfieldMesh/Public/VirtualHeightfieldMeshComponent.h
--- a/Engine/Plugins/Experimental/VirtualHeightfieldMesh/Source/VirtualHeightfieldMesh/Public/VirtualHeightfieldMeshComponent.h
+++ b/Engine/Plugins/Experimental/VirtualHeightfieldMesh/Source/VirtualHeightfieldMesh/Public/VirtualHeightfieldMeshComponent.h
@@ -62,6 +62,10 @@ protected:
 	UPROPERTY(EditAnywhere, Category = Rendering, meta = (DisplayName = "LOD Distribution", ClampMin = "1.0", UIMin = "1.0"))
 	float LodDistribution = 2.f;
 
+	/** Scale applied to LodBias texture. A larger value increases resolution where the LodBias texture is set. */
+	UPROPERTY(EditAnywhere, Category = Rendering, meta = (DisplayName = "LOD Bias Scale", ClampMin = "0.0", ClampMax = "10.0", UIMin = "0.0", UIMax = "10.0"))
+	float LodBiasScale = 0.f;
+
 	/** The number of levels of geometry subdivision to apply before the LOD 0 from the source virtual texture. */
 	UPROPERTY(EditAnywhere, Category = Rendering, meta = (DisplayName = "Subdivision LODs", ClampMin = "0", ClampMax = "8", UIMin = "0", UIMax = "8"))
 	int32 NumSubdivisionLods = 0;
@@ -111,6 +115,7 @@ public:
 	float GetLod0ScreenSize() const { return Lod0ScreenSize; }
 	float GetLod0Distribution() const { return Lod0Distribution; }
 	float GetLodDistribution() const { return LodDistribution; }
+	float GetLodBiasScale() const { return LodBiasScale; }
 	int32 GetNumSubdivisionLods() const { return NumSubdivisionLods; }
 	int32 GetNumTailLods() const { return NumTailLods; }
 	int32 GetNumForceLoadLods() const { return NumForceLoadLods; }
