Remove unused VHM options.

diff --git a/Engine/Plugins/Experimental/VirtualHeightfieldMesh/Shaders/Private/VirtualHeightfieldMesh.usf b/Engine/Plugins/Experimental/VirtualHeightfieldMesh/Shaders/Private/VirtualHeightfieldMesh.usf
--- a/Engine/Plugins/Experimental/VirtualHeightfieldMesh/Shaders/Private/VirtualHeightfieldMesh.usf
+++ b/Engine/Plugins/Experimental/VirtualHeightfieldMesh/Shaders/Private/VirtualHeightfieldMesh.usf
@@ -6,7 +6,6 @@
 //#include "/Engine/Private/ShaderDrawDebug.ush"	
 #include "VirtualHeightfieldMesh.ush"	
 
-
 groupshared uint NumGroupTasks;
 globallycoherent RWStructuredBuffer<WorkerQueueInfo> RWQueueInfo;
 globallycoherent RWBuffer<uint> RWQueueBuffer;
@@ -185,18 +184,11 @@ void InitBuffersCS()
 
 	RWQueueBuffer[0] = Pack(InitQuadItem(0, MaxLevel));
 
-	// RenderLodMap indirect args
-	RWIndirectArgsBuffer[0] = 6;
-	RWIndirectArgsBuffer[1] = 0; // Increment this counter during CollectQuadsCS.
-	RWIndirectArgsBuffer[2] = 0;
-	RWIndirectArgsBuffer[3] = 0;
-	RWIndirectArgsBuffer[4] = 0;
-
-	// ResolveNeighborLods and CullInstances indirect args
-	RWIndirectArgsBuffer[5] = 0; // Increment this counter during CollectQuadsCS.
-	RWIndirectArgsBuffer[6] = 1;
-	RWIndirectArgsBuffer[7] = 1;
-	RWIndirectArgsBuffer[8] = 0;
+	// CullInstances indirect args
+	RWIndirectArgsBuffer[0] = 0; // Increment this wave counter during CollectQuadsCS.
+	RWIndirectArgsBuffer[1] = 1;
+	RWIndirectArgsBuffer[2] = 1;
+	RWIndirectArgsBuffer[3] = 0; // Increment this instance counter counter during CollectQuadsCS.
 
 	// Clear virtual texture feedback counter or prime feedback buffer with first n levels.
 	uint FeedbackCount = 0;
@@ -317,8 +309,8 @@ void CollectQuadsCS(
 				uint PhysicalAddress = PageTableTexture.Load(int3(Pos, Level));
 
 				uint Write;
-				InterlockedAdd(RWIndirectArgsBuffer[1], 1, Write);
-				InterlockedMax(RWIndirectArgsBuffer[5], ((Write + 1) + 63) / 64);
+				InterlockedAdd(RWIndirectArgsBuffer[3], 1, Write);
+				InterlockedMax(RWIndirectArgsBuffer[0], ((Write + 1) + 63) / 64);
 
 				RWQuadBuffer[Write] = Pack(InitQuadRenderItem(Pos, Level, PhysicalAddress, bCull || bOcclude));
 
@@ -377,7 +369,7 @@ void InitInstanceBufferCS()
 void CullInstancesCS( uint3 DispatchThreadId : SV_DispatchThreadID )
 {
 	uint QuadIndex = DispatchThreadId.x;
-	if (QuadIndex >= IndirectArgsBufferSRV[1])
+	if (QuadIndex >= IndirectArgsBufferSRV[3])
 		return;
 
 	uint2 PackedItem = QuadBuffer[QuadIndex];
diff --git a/Engine/Plugins/Experimental/VirtualHeightfieldMesh/Source/VirtualHeightfieldMesh/Private/VirtualHeightfieldMeshSceneProxy.cpp b/Engine/Plugins/Experimental/VirtualHeightfieldMesh/Source/VirtualHeightfieldMesh/Private/VirtualHeightfieldMeshSceneProxy.cpp
--- a/Engine/Plugins/Experimental/VirtualHeightfieldMesh/Source/VirtualHeightfieldMesh/Private/VirtualHeightfieldMeshSceneProxy.cpp
+++ b/Engine/Plugins/Experimental/VirtualHeightfieldMesh/Source/VirtualHeightfieldMesh/Private/VirtualHeightfieldMeshSceneProxy.cpp
@@ -372,8 +372,6 @@ FVirtualHeightfieldMeshSceneProxy::FVirtualHeightfieldMeshSceneProxy(UVirtualHei
 	, Lod0Distribution(InComponent->GetLod0Distribution())
 	, LodDistribution(InComponent->GetLodDistribution())
 	, LodBiasScale(InComponent->GetLodBiasScale())
-	, NumSubdivisionLODs(InComponent->GetNumSubdivisionLods())
-	, NumTailLods(InComponent->GetNumTailLods())
 	, NumForceLoadLods(InComponent->GetNumForceLoadLods())
 	, NumOcclusionLods(0)
 	, OcclusionGridSize(0, 0)
@@ -543,7 +541,7 @@ void FVirtualHeightfieldMeshSceneProxy::GetDynamicMeshElements(const TArray<cons
 				const float PhysicalTextureSize = AllocatedVirtualTexture->GetPhysicalTextureSize(0);
 				UserData->PhysicalTextureSize = FVector2D(PhysicalTextureSize, 1.f / PhysicalTextureSize);
 
-				UserData->MaxLod = AllocatedVirtualTexture->GetMaxLevel() + NumTailLods;
+				UserData->MaxLod = AllocatedVirtualTexture->GetMaxLevel();
 				UserData->VirtualHeightfieldToLocal = UVToLocal;
 				UserData->VirtualHeightfieldToWorld = UVToWorld;
 
@@ -645,7 +643,7 @@ void FVirtualHeightfieldMeshSceneProxy::AcceptOcclusionResults(FSceneView const*
 		OcclusionResults.TextureSize = OcclusionGridSize;
 		OcclusionResults.NumTextureMips = NumOcclusionLods;
 		
-		FRHIResourceCreateInfo CreateInfo(TEXT("FVirtualHeightfieldMeshSceneProxy_OcclusionTexture"));
+		FRHIResourceCreateInfo CreateInfo(TEXT("VirtualHeightfieldMesh.OcclusionTexture"));
 		OcclusionResults.OcclusionTexture = RHICreateTexture2D(OcclusionGridSize.X, OcclusionGridSize.Y, PF_G8, NumOcclusionLods, 1, TexCreate_ShaderResource, CreateInfo);
 		
 		bool const* Src = Results->GetData() + ResultsStart;
@@ -674,10 +672,8 @@ void FVirtualHeightfieldMeshSceneProxy::AcceptOcclusionResults(FSceneView const*
 namespace VirtualHeightfieldMesh
 {
 	/* Keep indirect args offsets in sync with VirtualHeightfieldMesh.usf. */
-	static const int32 IndirectArgsByteOffset_RenderLodMap = 0;
-	static const int32 IndirectArgsByteOffset_FetchNeighborLod = 5 * sizeof(uint32);
-	static const int32 IndirectArgsByteOffset_FinalCull = 5 * sizeof(uint32);
-	static const int32 IndirectArgsByteSize = 9 * sizeof(uint32);
+	static const int32 IndirectArgsByteOffset_FinalCull = 0;
+	static const int32 IndirectArgsByteSize = 4 * sizeof(uint32);
 
 	/** Shader structure used for tracking work queues in persistent wave style shaders. Keep in sync with VirtualHeightfieldMesh.ush. */
 	struct WorkerQueueInfo
@@ -839,7 +835,7 @@ namespace VirtualHeightfieldMesh
 	public:
 		virtual void InitRHI() override
 		{
-			FRHIResourceCreateInfo CreateInfo(TEXT("MinMaxDefaultTexture"));
+			FRHIResourceCreateInfo CreateInfo(TEXT("VirtualHeightfieldMesh.MinMaxDefaultTexture"));
 			FTexture2DRHIRef Texture2D = RHICreateTexture2D(1, 1, PF_B8G8R8A8, 1, 1, TexCreate_ShaderResource, CreateInfo);
 			TextureRHI = Texture2D;
 
@@ -986,7 +982,7 @@ namespace VirtualHeightfieldMesh
 	void InitializeInstanceBuffers(FRHICommandListImmediate& InRHICmdList, FDrawInstanceBuffers& InBuffers)
 	{
 		{
-			FRHIResourceCreateInfo CreateInfo(TEXT("FVirtualHeightfieldMeshSceneProxy_InstanceBuffer"));
+			FRHIResourceCreateInfo CreateInfo(TEXT("VirtualHeightfieldMesh.InstanceBuffer"));
 			const int32 InstanceSize = sizeof(VirtualHeightfieldMesh::QuadRenderInstance);
 			const int32 InstanceBufferSize = CVarVHMMaxRenderItems.GetValueOnRenderThread() * InstanceSize;
 			InBuffers.InstanceBuffer = RHICreateStructuredBuffer(InstanceSize, InstanceBufferSize, BUF_UnorderedAccess|BUF_ShaderResource, ERHIAccess::SRVMask, CreateInfo);
@@ -994,7 +990,7 @@ namespace VirtualHeightfieldMesh
 			InBuffers.InstanceBufferSRV = RHICreateShaderResourceView(InBuffers.InstanceBuffer);
 		}
 		{
-			FRHIResourceCreateInfo CreateInfo(TEXT("FVirtualHeightfieldMeshSceneProxy_IndirectArgsBuffer"));
+			FRHIResourceCreateInfo CreateInfo(TEXT("VirtualHeightfieldMesh.InstanceIndirectArgsBuffer"));
 			InBuffers.IndirectArgsBuffer = RHICreateVertexBuffer(5 * sizeof(uint32), BUF_UnorderedAccess|BUF_DrawIndirect, ERHIAccess::IndirectArgs, CreateInfo);
 			InBuffers.IndirectArgsBufferUAV = RHICreateUnorderedAccessView(InBuffers.IndirectArgsBuffer, PF_R32_UINT);
 		}
@@ -1003,21 +999,21 @@ namespace VirtualHeightfieldMesh
 	/** Initialize the volatile resources used in the render graph. */
 	void InitializeResources(FRDGBuilder& GraphBuilder, FProxyDesc const& InDesc, FMainViewDesc const& InMainViewDesc, FVolatileResources& OutResources)
 	{
-		OutResources.QueueInfo = GraphBuilder.CreateBuffer(FRDGBufferDesc::CreateStructuredDesc(sizeof(WorkerQueueInfo), 1), TEXT("QueueInfo"));
+		OutResources.QueueInfo = GraphBuilder.CreateBuffer(FRDGBufferDesc::CreateStructuredDesc(sizeof(WorkerQueueInfo), 1), TEXT("VirtualHeightfieldMesh.QueueInfo"));
 		OutResources.QueueInfoUAV = GraphBuilder.CreateUAV(OutResources.QueueInfo);
-		OutResources.QueueBuffer = GraphBuilder.CreateBuffer(FRDGBufferDesc::CreateBufferDesc(sizeof(uint32), InDesc.MaxPersistentQueueItems), TEXT("QuadQueue"));
+		OutResources.QueueBuffer = GraphBuilder.CreateBuffer(FRDGBufferDesc::CreateBufferDesc(sizeof(uint32), InDesc.MaxPersistentQueueItems), TEXT("VirtualHeightfieldMesh.QuadQueue"));
 		OutResources.QueueBufferUAV = GraphBuilder.CreateUAV(FRDGBufferUAVDesc(OutResources.QueueBuffer, PF_R32_UINT));
 
-		OutResources.QuadBuffer = GraphBuilder.CreateBuffer(FRDGBufferDesc::CreateBufferDesc(sizeof(uint32) * 2, InDesc.MaxRenderItems), TEXT("QuadBuffer"));
+		OutResources.QuadBuffer = GraphBuilder.CreateBuffer(FRDGBufferDesc::CreateBufferDesc(sizeof(uint32) * 2, InDesc.MaxRenderItems), TEXT("VirtualHeightfieldMesh.QuadBuffer"));
 		OutResources.QuadBufferUAV = GraphBuilder.CreateUAV(FRDGBufferUAVDesc(OutResources.QuadBuffer, PF_R32G32_UINT));
 		OutResources.QuadBufferSRV = GraphBuilder.CreateSRV(FRDGBufferSRVDesc(OutResources.QuadBuffer, PF_R32G32_UINT));
 
 		FRDGBufferDesc FeedbackBufferDesc = FRDGBufferDesc::CreateBufferDesc(sizeof(uint32), InDesc.MaxFeedbackItems + 1);
 		FeedbackBufferDesc.Usage = EBufferUsageFlags(FeedbackBufferDesc.Usage | BUF_SourceCopy);
-		OutResources.FeedbackBuffer = GraphBuilder.CreateBuffer(FeedbackBufferDesc, TEXT("FeedbackBuffer"));
+		OutResources.FeedbackBuffer = GraphBuilder.CreateBuffer(FeedbackBufferDesc, TEXT("VirtualHeightfieldMesh.FeedbackBuffer"));
 		OutResources.FeedbackBufferUAV = GraphBuilder.CreateUAV(FRDGBufferUAVDesc(OutResources.FeedbackBuffer, PF_R32_UINT));
 
-		OutResources.IndirectArgsBuffer = GraphBuilder.CreateBuffer(FRDGBufferDesc::CreateIndirectDesc(IndirectArgsByteSize), TEXT("IndirectArgsBuffer"));
+		OutResources.IndirectArgsBuffer = GraphBuilder.CreateBuffer(FRDGBufferDesc::CreateIndirectDesc(IndirectArgsByteSize), TEXT("VirtualHeightfieldMesh.IndirectArgsBuffer"));
 		OutResources.IndirectArgsBufferUAV = GraphBuilder.CreateUAV(OutResources.IndirectArgsBuffer);
 		OutResources.IndirectArgsBufferSRV = GraphBuilder.CreateSRV(OutResources.IndirectArgsBuffer);
 	}
diff --git a/Engine/Plugins/Experimental/VirtualHeightfieldMesh/Source/VirtualHeightfieldMesh/Private/VirtualHeightfieldMeshSceneProxy.h b/Engine/Plugins/Experimental/VirtualHeightfieldMesh/Source/VirtualHeightfieldMesh/Private/VirtualHeightfieldMeshSceneProxy.h
--- a/Engine/Plugins/Experimental/VirtualHeightfieldMesh/Source/VirtualHeightfieldMesh/Private/VirtualHeightfieldMeshSceneProxy.h
+++ b/Engine/Plugins/Experimental/VirtualHeightfieldMesh/Source/VirtualHeightfieldMesh/Private/VirtualHeightfieldMeshSceneProxy.h
@@ -57,8 +57,6 @@ public:
 	float LodDistribution;
 	float LodBiasScale;
 	
-	int32 NumSubdivisionLODs;
-	int32 NumTailLods;
 	int32 NumForceLoadLods;
 
 	TArray<FVector2D> OcclusionData;
diff --git a/Engine/Plugins/Experimental/VirtualHeightfieldMesh/Source/VirtualHeightfieldMesh/Public/VirtualHeightfieldMeshComponent.h b/Engine/Plugins/Experimental/VirtualHeightfieldMesh/Source/VirtualHeightfieldMesh/Public/VirtualHeightfieldMeshComponent.h
--- a/Engine/Plugins/Experimental/VirtualHeightfieldMesh/Source/VirtualHeightfieldMesh/Public/VirtualHeightfieldMeshComponent.h
+++ b/Engine/Plugins/Experimental/VirtualHeightfieldMesh/Source/VirtualHeightfieldMesh/Public/VirtualHeightfieldMeshComponent.h
@@ -50,7 +50,7 @@ protected:
 	UPROPERTY(EditAnywhere, Category = Rendering)
 	UMaterialInterface* Material = nullptr;
 
-	/** Target sceen size for LOD 0. A larger value uniformly increases the geometry resolution on screen. */
+	/** Target screen size for a LOD 0 tile. A smaller value uniformly increases the geometry resolution on screen. */
 	UPROPERTY(EditAnywhere, Category = Rendering, meta = (DisplayName = "LOD 0 Screen Size", ClampMin = "0.1", UIMin = "0.1"))
 	float Lod0ScreenSize = 1.f;
 
@@ -59,21 +59,13 @@ protected:
 	float Lod0Distribution = 1.f;
 
 	/** Distribution multiplier applied for each LOD level. A larger value increases the distance exponentially between each LOD transition. */
-	UPROPERTY(EditAnywhere, Category = Rendering, meta = (DisplayName = "LOD Distribution", ClampMin = "1.0", UIMin = "1.0"))
+	UPROPERTY(EditAnywhere, Category = Rendering, meta = (DisplayName = "LOD Distribution", ClampMin = "1.0", UIMin = "1.0", UIMax = "3.0"))
 	float LodDistribution = 2.f;
 
 	/** Scale applied to LodBias texture. A larger value increases resolution where the LodBias texture is set. */
 	UPROPERTY(EditAnywhere, Category = Rendering, meta = (DisplayName = "LOD Bias Scale", ClampMin = "0.0", ClampMax = "10.0", UIMin = "0.0", UIMax = "10.0"))
 	float LodBiasScale = 0.f;
 
-	/** The number of levels of geometry subdivision to apply before the LOD 0 from the source virtual texture. */
-	UPROPERTY(EditAnywhere, Category = Rendering, meta = (DisplayName = "Subdivision LODs", ClampMin = "0", ClampMax = "8", UIMin = "0", UIMax = "8"))
-	int32 NumSubdivisionLods = 0;
-
-	/** The number of levels of geometry reduction to apply after the Max LOD from the source virtual texture. */
-	UPROPERTY(EditAnywhere, Category = Rendering, meta = (DisplayName = "Tail LODs", ClampMin = "0", ClampMax = "8", UIMin = "0", UIMax = "8"))
-	int32 NumTailLods = 0;
-
 	/** The number of levels that we force loaded. A higher number can reduce LOD pop in but allocates more of the virtual texture memory pool. */
 	UPROPERTY(EditAnywhere, Category = Rendering, meta = (DisplayName = "Force Load LODs", ClampMin = "0", ClampMax = "4"))
 	int32 NumForceLoadLods = 0;
@@ -116,8 +108,6 @@ public:
 	float GetLod0Distribution() const { return Lod0Distribution; }
 	float GetLodDistribution() const { return LodDistribution; }
 	float GetLodBiasScale() const { return LodBiasScale; }
-	int32 GetNumSubdivisionLods() const { return NumSubdivisionLods; }
-	int32 GetNumTailLods() const { return NumTailLods; }
 	int32 GetNumForceLoadLods() const { return NumForceLoadLods; }
 	int32 GetNumOcclusionLods() const { return NumOcclusionLods; }
 
