Add a world partition builder implementation for VirtualHeightfieldMesh.

diff --git a/Engine/Plugins/Experimental/VirtualHeightfieldMesh/Source/VirtualHeightfieldMeshEditor/Private/WorldPartitionVirtualHeightfieldMeshBuilder.cpp b/Engine/Plugins/Experimental/VirtualHeightfieldMesh/Source/VirtualHeightfieldMeshEditor/Private/WorldPartitionVirtualHeightfieldMeshBuilder.cpp
new file mode 100644
--- /dev/null
+++ b/Engine/Plugins/Experimental/VirtualHeightfieldMesh/Source/VirtualHeightfieldMeshEditor/Private/WorldPartitionVirtualHeightfieldMeshBuilder.cpp
@@ -0,0 +1,65 @@
+// Copyright Epic Games, Inc. All Rights Reserved.
+
+#include "WorldPartitionVirtualHeightfieldMeshBuilder.h"
+
+#include "HeightfieldMinMaxTexture.h"
+#include "HeightfieldMinMaxTextureBuild.h"
+#include "UObject/UObjectIterator.h"
+#include "VirtualHeightfieldMeshComponent.h"
+
+DEFINE_LOG_CATEGORY_STATIC(LogWorldPartitionVirtualHeightfieldMeshBuilder, All, All);
+
+UWorldPartitionVirtualHeightfieldMeshBuilder::UWorldPartitionVirtualHeightfieldMeshBuilder(const FObjectInitializer& ObjectInitializer)
+	: Super(ObjectInitializer)
+{
+}
+
+bool UWorldPartitionVirtualHeightfieldMeshBuilder::Run(UWorld* World, FPackageSourceControlHelper& PackageHelper)
+{
+	// Find UVirtualHeightfieldMeshComponent and update the associated UHeightfieldMinMaxTexture assets.
+	// todo[vhm]: Convert builder to sequentially load and build sections of the world.
+	TSet<UObject*> ModifiedObjects;
+	for (TObjectIterator<UVirtualHeightfieldMeshComponent> It(RF_ClassDefaultObject, false, EInternalObjectFlags::PendingKill); It; ++It)
+	{
+		if (It->GetWorld() == World)
+		{
+			if (VirtualHeightfieldMesh::HasMinMaxHeightTexture(*It) && VirtualHeightfieldMesh::BuildMinMaxHeightTexture(*It))
+			{
+				ModifiedObjects.Add(It->GetMinMaxTexture());
+			}
+		}
+	}
+
+	// Checkout and save any modified packages.
+	for (UObject* ModifiedObject : ModifiedObjects)
+	{
+		UObject* Outer = ModifiedObject->GetOuter();
+		UPackage* Package = Cast<UPackage>(Outer);
+
+		if (Package != nullptr && Package->IsDirty())
+		{
+			if (!PackageHelper.Checkout(Package))
+			{
+				UE_LOG(LogWorldPartitionVirtualHeightfieldMeshBuilder, Error, TEXT("Error checking out package %s."), *Package->GetName());
+				return false;
+			}
+
+			FString PackageFileName = SourceControlHelpers::PackageFilename(Package);
+			if (!UPackage::SavePackage(Package, nullptr, RF_Standalone, *PackageFileName, GError, nullptr, false, true, SAVE_Async))
+			{
+				UE_LOG(LogWorldPartitionVirtualHeightfieldMeshBuilder, Error, TEXT("Error saving package %s."), *Package->GetName());
+				return false;
+			}
+
+			if (!PackageHelper.AddToSourceControl(Package))
+			{
+				UE_LOG(LogWorldPartitionVirtualHeightfieldMeshBuilder, Error, TEXT("Error adding package %s to source control."), *Package->GetName());
+				return false;
+			}
+		}
+	}
+
+	UPackage::WaitForAsyncFileWrites();
+
+	return true;
+}
diff --git a/Engine/Plugins/Experimental/VirtualHeightfieldMesh/Source/VirtualHeightfieldMeshEditor/Private/WorldPartitionVirtualHeightfieldMeshBuilder.h b/Engine/Plugins/Experimental/VirtualHeightfieldMesh/Source/VirtualHeightfieldMeshEditor/Private/WorldPartitionVirtualHeightfieldMeshBuilder.h
new file mode 100644
--- /dev/null
+++ b/Engine/Plugins/Experimental/VirtualHeightfieldMesh/Source/VirtualHeightfieldMeshEditor/Private/WorldPartitionVirtualHeightfieldMeshBuilder.h
@@ -0,0 +1,19 @@
+// Copyright Epic Games, Inc. All Rights Reserved.
+
+#pragma once
+
+#include "WorldPartition/WorldPartitionBuilder.h"
+#include "WorldPartitionVirtualHeightfieldMeshBuilder.generated.h"
+
+UCLASS()
+class UWorldPartitionVirtualHeightfieldMeshBuilder : public UWorldPartitionBuilder
+{
+	GENERATED_UCLASS_BODY()
+
+public:
+	// UWorldPartitionBuilder interface begin
+	virtual bool RequiresCommandletRendering() const override { return true; }
+	virtual bool RequiresEntireWorldLoading() const override { return true; }
+	virtual bool Run(UWorld* World, FPackageSourceControlHelper& PackageHelper) override;
+	// UWorldPartitionBuilder interface end
+};
\ No newline at end of file
diff --git a/Engine/Plugins/Experimental/VirtualHeightfieldMesh/Source/VirtualHeightfieldMeshEditor/VirtualHeightfieldMeshEditor.Build.cs b/Engine/Plugins/Experimental/VirtualHeightfieldMesh/Source/VirtualHeightfieldMeshEditor/VirtualHeightfieldMeshEditor.Build.cs
--- a/Engine/Plugins/Experimental/VirtualHeightfieldMesh/Source/VirtualHeightfieldMeshEditor/VirtualHeightfieldMeshEditor.Build.cs
+++ b/Engine/Plugins/Experimental/VirtualHeightfieldMesh/Source/VirtualHeightfieldMeshEditor/VirtualHeightfieldMeshEditor.Build.cs
@@ -14,6 +14,7 @@ public class VirtualHeightfieldMeshEditor : ModuleRules
 			"RenderCore",
 			"Renderer",
 			"RHI",
+			"SourceControl",
 			"Slate",
 			"SlateCore",
 			"UnrealEd",
