VHM: Move immutable state to uniform buffer.
Fix vertex shader using wrong SRV of physical height leading to unwanted sRGB transform of height data on some platforms.

diff --git a/Engine/Plugins/Experimental/VirtualHeightfieldMesh/Shaders/Private/VirtualHeightfieldMeshVertexFactory.ush b/Engine/Plugins/Experimental/VirtualHeightfieldMesh/Shaders/Private/VirtualHeightfieldMeshVertexFactory.ush
--- a/Engine/Plugins/Experimental/VirtualHeightfieldMesh/Shaders/Private/VirtualHeightfieldMeshVertexFactory.ush
+++ b/Engine/Plugins/Experimental/VirtualHeightfieldMesh/Shaders/Private/VirtualHeightfieldMeshVertexFactory.ush
@@ -4,30 +4,11 @@
 #include "/Engine/Private/VirtualTextureCommon.ush"
 #include "VirtualHeightfieldMesh.ush"	
 
-StructuredBuffer<QuadRenderInstance> InstanceBuffer;
-
-Texture2D<uint4> PageTableTexture;
-Texture2D<float4> HeightTexture;
-SamplerState HeightSampler;
-
-uint4 VTPackedUniform;
-uint4 VTPackedPageTableUniform0;
-uint4 VTPackedPageTableUniform1;
-
-float4 PageTableSize;
-float2 PhysicalTextureSize;
-float MaxLod;
-float4x4 VirtualHeightfieldToLocal;
-float4x4 VirtualHeightfieldToWorld;
-
-Texture2D<float> LodBiasTexture;
-SamplerState LodBiasSampler;
+#define GRID_SIZE (VHM.NumQuadsPerTileSide+1)
 
+StructuredBuffer<QuadRenderInstance> InstanceBuffer;
 float3 LodViewOrigin;
 float4 LodDistances;
-float LodBiasScale;
-
-#define GRID_SIZE (VirtualHeightfieldMeshVF.NumQuadsPerTileSide+1)
 
 /** Per-vertex inputs. No vertex buffers are bound. */
 struct FVertexFactoryInput
@@ -60,7 +41,7 @@ struct FVertexFactoryInterpolantsVSToPS
 /** Fetch the lod bias for the location. */
 float GetLodBias(float2 InUV, float InLodBiasScale)
 {
-	return CalculateBiasLod(LodBiasTexture.SampleLevel(LodBiasSampler, InUV, 0), InLodBiasScale);
+	return CalculateBiasLod(VHM.LodBiasTexture.SampleLevel(VHM.LodBiasSampler, InUV, 0), InLodBiasScale);
 }
 
 /** Helper to morph the UV location of a vertex. First snaps down InMorphFactorFloor LOD levels and then morphs InMorphFactorFrac towards the next LOD level. */
@@ -97,21 +78,21 @@ FVertexFactoryIntermediates GetVertexFactoryIntermediates(FVertexFactoryInput In
 
 	// Calculate vertex UV details before morphing
 	float2 XY = ((float2)Pos + LocalUV) * (float)(1 << Level);
-	float2 NormalizedPos = (XY * PageTableSize.zw);
+	float2 NormalizedPos = (XY * VHM.PageTableSize.zw);
 	float SampleLevel = Level;
 
 	// Sample height once to approximate distance and morph the LocalUV.
 	{
 		float2 LocalPhysicalUV = LocalUVTransform.xy + LocalUV * LocalUVTransform.z;
-		float Height = HeightTexture.SampleLevel(HeightSampler, LocalPhysicalUV, 0);
-		float3 WorldPos = mul(float4(NormalizedPos, Height, 1), VirtualHeightfieldToWorld).xyz;
+		float Height = VHM.HeightTexture.SampleLevel(VHM.HeightSampler, LocalPhysicalUV, 0);
+		float3 WorldPos = mul(float4(NormalizedPos, Height, 1), VHM.VirtualHeightfieldToWorld).xyz;
 		float DistanceSq = dot(LodViewOrigin - WorldPos, LodViewOrigin - WorldPos);
 		float LodForDistance = CalculateDistanceLod(DistanceSq, LodDistances);
-		float LodBias = GetLodBias(NormalizedPos, LodBiasScale);
+		float LodBias = GetLodBias(NormalizedPos, VHM.LodBiasScale);
 		
 		// Clamp between the LOD level for this instance and the max LOD. 
 		// Note that the culling phase should already ensure that LOD >= Level.
-		float LodClamped = clamp(LodForDistance - LodBias, (float)Level, MaxLod);
+		float LodClamped = clamp(LodForDistance - LodBias, (float)Level, VHM.MaxLod);
 		
 		// Number of levels that we need to morph.
 		float LodMorphFloor = floor(LodClamped) - (float)Level;
@@ -128,27 +109,29 @@ FVertexFactoryIntermediates GetVertexFactoryIntermediates(FVertexFactoryInput In
 		LocalUV = MorphVertex(LocalUV, GRID_SIZE - 1, (uint)LodMorphFloor, LodMorphFrac);
 		// Adjust other position variables to match the new morphed position.
 		XY = ((float2)Pos + LocalUV) * (float)(1 << Level);
-		NormalizedPos = (XY * PageTableSize.zw);
+		NormalizedPos = (XY * VHM.PageTableSize.zw);
 
 		// Use LOD level as sample level but bias to keep sample level in a good range.
 		SampleLevel = max(0, LodClamped - 0.5f);
 	}
 	
 	// Sample height from virtual texture.
-	VTUniform Uniform = VTUniform_Unpack(VTPackedUniform);
-	Uniform.vPageBorderSize -= .5f * PhysicalTextureSize.y; // Half texel offset is used in VT write and in sampling because we want texel locations to match landscape vertices.
-	VTPageTableUniform PageTableUniform = VTPageTableUniform_Unpack(VTPackedPageTableUniform0, VTPackedPageTableUniform1);
-	VTPageTableResult VTResult0 = TextureLoadVirtualPageTableLevel(PageTableTexture, PageTableUniform, NormalizedPos, VTADDRESSMODE_CLAMP, VTADDRESSMODE_CLAMP, floor(SampleLevel));
-	float4 Height0 = TextureVirtualSampleLevel(HeightTexture, HeightSampler, VTResult0, 0, Uniform);
-	VTPageTableResult VTResult1 = TextureLoadVirtualPageTableLevel(PageTableTexture, PageTableUniform, NormalizedPos, VTADDRESSMODE_CLAMP, VTADDRESSMODE_CLAMP, ceil(SampleLevel));
-	float4 Height1 = TextureVirtualSampleLevel(HeightTexture, HeightSampler, VTResult1, 0, Uniform);
+	VTUniform Uniform = VTUniform_Unpack(VHM.VTPackedUniform);
+	Uniform.vPageBorderSize -= .5f * VHM.PhysicalTextureSize.y; // Half texel offset is used in VT write and in sampling because we want texel locations to match landscape vertices.
+	VTPageTableUniform PageTableUniform = VTPageTableUniform_Unpack(VHM.VTPackedPageTableUniform0, VHM.VTPackedPageTableUniform1);
+	VTPageTableResult VTResult0 = TextureLoadVirtualPageTableLevel(VHM.PageTableTexture, PageTableUniform, NormalizedPos, VTADDRESSMODE_CLAMP, VTADDRESSMODE_CLAMP, floor(SampleLevel));
+	float2 UV0 = VTComputePhysicalUVs(VTResult0, 0, Uniform);
+	float Height0 = VHM.HeightTexture.SampleLevel(VHM.HeightSampler, UV0, 0);
+	VTPageTableResult VTResult1 = TextureLoadVirtualPageTableLevel(VHM.PageTableTexture, PageTableUniform, NormalizedPos, VTADDRESSMODE_CLAMP, VTADDRESSMODE_CLAMP, ceil(SampleLevel));
+	float2 UV1 = VTComputePhysicalUVs(VTResult1, 0, Uniform);
+	float Height1 = VHM.HeightTexture.SampleLevel(VHM.HeightSampler, UV1, 0);
 	float Height = lerp(Height0.x, Height1.x, frac(SampleLevel));
 
 	// Position in space of virtual texture volume
 	Intermediates.VTPos = float3(NormalizedPos, Height);
 
 	// Position in local space
-	Intermediates.LocalPos = mul(float4(Intermediates.VTPos, 1), VirtualHeightfieldToLocal).xyz;
+	Intermediates.LocalPos = mul(float4(Intermediates.VTPos, 1), VHM.VirtualHeightfieldToLocal).xyz;
 	
 	Intermediates.LocalUV = NormalizedPos;
 
@@ -167,7 +150,7 @@ FMaterialVertexParameters GetMaterialVertexParameters(FVertexFactoryInput Input,
 	Result.WorldPosition = WorldPosition;
 
 	// needs fixing!
-	Result.TangentToWorld = mul(TangentToLocal, (float3x3)VirtualHeightfieldToWorld);
+	Result.TangentToWorld = mul(TangentToLocal, (float3x3)VHM.VirtualHeightfieldToWorld);
 	Result.TangentToWorld[2] = Intermediates.WorldNormal;
 
 	Result.PreSkinnedPosition = WorldPosition;// Intermediates.WorldPosPreDisplacement.xyz;
diff --git a/Engine/Plugins/Experimental/VirtualHeightfieldMesh/Source/VirtualHeightfieldMesh/Private/VirtualHeightfieldMeshSceneProxy.cpp b/Engine/Plugins/Experimental/VirtualHeightfieldMesh/Source/VirtualHeightfieldMesh/Private/VirtualHeightfieldMeshSceneProxy.cpp
--- a/Engine/Plugins/Experimental/VirtualHeightfieldMesh/Source/VirtualHeightfieldMesh/Private/VirtualHeightfieldMeshSceneProxy.cpp
+++ b/Engine/Plugins/Experimental/VirtualHeightfieldMesh/Source/VirtualHeightfieldMesh/Private/VirtualHeightfieldMeshSceneProxy.cpp
@@ -437,10 +437,44 @@ void FVirtualHeightfieldMeshSceneProxy::CreateRenderThreadResources()
 		if (RuntimeVirtualTexture->GetMaterialType() == ERuntimeVirtualTextureMaterialType::WorldHeight)
 		{
 			AllocatedVirtualTexture = RuntimeVirtualTexture->GetAllocatedVirtualTexture();
-
 			NumQuadsPerTileSide = RuntimeVirtualTexture->GetTileSize();
-			VertexFactory = new FVirtualHeightfieldMeshVertexFactory(GetScene().GetFeatureLevel(), NumQuadsPerTileSide);
-			VertexFactory->InitResource();
+
+			if (AllocatedVirtualTexture != nullptr)
+			{
+				// Gather vertex factory uniform parameters.
+				FVirtualHeightfieldMeshVertexFactoryParameters UniformParams;
+				UniformParams.PageTableTexture = AllocatedVirtualTexture->GetPageTableTexture(0);
+				UniformParams.HeightTexture = AllocatedVirtualTexture->GetPhysicalTextureSRV(0, false);
+				UniformParams.HeightSampler = TStaticSamplerState<SF_Bilinear>::GetRHI();
+				UniformParams.LodBiasTexture = LodBiasTexture ? LodBiasTexture->Resource->TextureRHI : GBlackTexture->TextureRHI;
+				UniformParams.LodBiasSampler = TStaticSamplerState<SF_Point>::GetRHI();
+				UniformParams.NumQuadsPerTileSide = NumQuadsPerTileSide;
+
+				FUintVector4 PackedUniform;
+				AllocatedVirtualTexture->GetPackedUniform(&PackedUniform, 0);
+				UniformParams.VTPackedUniform = PackedUniform;
+				FUintVector4 PackedPageTableUniform[2];
+				AllocatedVirtualTexture->GetPackedPageTableUniform(PackedPageTableUniform);
+				UniformParams.VTPackedPageTableUniform0 = PackedPageTableUniform[0];
+				UniformParams.VTPackedPageTableUniform1 = PackedPageTableUniform[1];
+
+				const float PageTableSizeX = AllocatedVirtualTexture->GetWidthInTiles();
+				const float PageTableSizeY = AllocatedVirtualTexture->GetHeightInTiles();
+				UniformParams.PageTableSize = FVector4(PageTableSizeX, PageTableSizeY, 1.f / PageTableSizeX, 1.f / PageTableSizeY);
+
+				const float PhysicalTextureSize = AllocatedVirtualTexture->GetPhysicalTextureSize(0);
+				UniformParams.PhysicalTextureSize = FVector2D(PhysicalTextureSize, 1.f / PhysicalTextureSize);
+
+				UniformParams.VirtualHeightfieldToLocal = UVToLocal;
+				UniformParams.VirtualHeightfieldToWorld = UVToWorld;
+
+				UniformParams.MaxLod = AllocatedVirtualTexture->GetMaxLevel();
+				UniformParams.LodBiasScale = LodBiasScale;
+
+				// Create vertex factory.
+				VertexFactory = new FVirtualHeightfieldMeshVertexFactory(GetScene().GetFeatureLevel(), UniformParams);
+				VertexFactory->InitResource();
+			}
 		}
 	}
 }
@@ -514,7 +548,7 @@ void FVirtualHeightfieldMeshSceneProxy::GetDynamicMeshElements(const TArray<cons
 			{
 				FMeshBatchElement& BatchElement = Mesh.Elements[0];
 
-				BatchElement.IndexBuffer = VertexFactory->IndexBuffer;
+				BatchElement.IndexBuffer = VertexFactory->GetIndexBuffer();
 				BatchElement.IndirectArgsBuffer = Buffers.IndirectArgsBuffer;
 				BatchElement.IndirectArgsOffset = 0;
 
@@ -523,30 +557,16 @@ void FVirtualHeightfieldMeshSceneProxy::GetDynamicMeshElements(const TArray<cons
 				BatchElement.MinVertexIndex = 0;
 				BatchElement.MaxVertexIndex = 0;
 
+				BatchElement.PrimitiveIdMode = PrimID_ForceZero;
+				BatchElement.PrimitiveUniformBuffer = GetUniformBuffer();
+
 				FVirtualHeightfieldMeshUserData* UserData = &Collector.AllocateOneFrameResource<FVirtualHeightfieldMeshUserData>();
 				BatchElement.UserData = (void*)UserData;
 
 				UserData->InstanceBufferSRV = Buffers.InstanceBufferSRV;
-				UserData->PageTableTexture = AllocatedVirtualTexture->GetPageTableTexture(0);
-				UserData->HeightPhysicalTexture = AllocatedVirtualTexture->GetPhysicalTexture(0);
-				UserData->LodBiasTexture = LodBiasTexture ? LodBiasTexture->Resource->TextureRHI : GBlackTexture->TextureRHI;
-
-				AllocatedVirtualTexture->GetPackedUniform(&UserData->PackedUniform, 0);
-				AllocatedVirtualTexture->GetPackedPageTableUniform(UserData->PackedPageTableUniform);
-
-				const float PageTableSizeX = AllocatedVirtualTexture->GetWidthInTiles();
-				const float PageTableSizeY = AllocatedVirtualTexture->GetHeightInTiles();
-				UserData->PageTableSize = FVector4(PageTableSizeX, PageTableSizeY, 1.f / PageTableSizeX, 1.f / PageTableSizeY);
-
-				const float PhysicalTextureSize = AllocatedVirtualTexture->GetPhysicalTextureSize(0);
-				UserData->PhysicalTextureSize = FVector2D(PhysicalTextureSize, 1.f / PhysicalTextureSize);
-
-				UserData->MaxLod = AllocatedVirtualTexture->GetMaxLevel();
-				UserData->VirtualHeightfieldToLocal = UVToLocal;
-				UserData->VirtualHeightfieldToWorld = UVToWorld;
 
+				//todo[vhm]: Move all the view dependent lod logic into shader. Would help us to move to static mesh batches in the future.
 				FSceneView const* MainView = ViewFamily.Views[0];
-				
 				UserData->LodViewOrigin = MainView->ViewMatrices.GetViewOrigin();
 
 #if !(UE_BUILD_SHIPPING || UE_BUILD_TEST)
@@ -559,11 +579,8 @@ void FVirtualHeightfieldMeshSceneProxy::GetDynamicMeshElements(const TArray<cons
 #endif
 				
 				UserData->LodDistances = VirtualHeightfieldMesh::CalculateLodRanges(MainView, this);
-				UserData->LodBiasScale = LodBiasScale;
-				
-				BatchElement.PrimitiveIdMode = PrimID_ForceZero;
-				BatchElement.PrimitiveUniformBuffer = GetUniformBuffer();
 			}
+
 			Collector.AddMesh(ViewIndex, Mesh);
 		}
 	}
diff --git a/Engine/Plugins/Experimental/VirtualHeightfieldMesh/Source/VirtualHeightfieldMesh/Private/VirtualHeightfieldMeshVertexFactory.cpp b/Engine/Plugins/Experimental/VirtualHeightfieldMesh/Source/VirtualHeightfieldMesh/Private/VirtualHeightfieldMeshVertexFactory.cpp
--- a/Engine/Plugins/Experimental/VirtualHeightfieldMesh/Source/VirtualHeightfieldMesh/Private/VirtualHeightfieldMeshVertexFactory.cpp
+++ b/Engine/Plugins/Experimental/VirtualHeightfieldMesh/Source/VirtualHeightfieldMesh/Private/VirtualHeightfieldMeshVertexFactory.cpp
@@ -10,7 +10,7 @@
 #include "ShaderParameters.h"
 #include "ShaderParameterUtils.h"
 
-IMPLEMENT_GLOBAL_SHADER_PARAMETER_STRUCT(FVirtualHeightfieldMeshVertexFactoryParameters, "VirtualHeightfieldMeshVF");
+IMPLEMENT_GLOBAL_SHADER_PARAMETER_STRUCT(FVirtualHeightfieldMeshVertexFactoryParameters, "VHM");
 
 namespace
 {
@@ -100,22 +100,8 @@ public:
 	void Bind(const FShaderParameterMap& ParameterMap)
 	{
 		InstanceBufferParameter.Bind(ParameterMap, TEXT("InstanceBuffer"));
-		PageTableTextureParameter.Bind(ParameterMap, TEXT("PageTableTexture"));
-		HeightTextureParameter.Bind(ParameterMap, TEXT("HeightTexture"));
-		HeightSamplerParameter.Bind(ParameterMap, TEXT("HeightSampler"));
-		LodBiasTextureParameter.Bind(ParameterMap, TEXT("LodBiasTexture"));
-		LodBiasSamplerParameter.Bind(ParameterMap, TEXT("LodBiasSampler"));
-		PackedUniformParameter.Bind(ParameterMap, TEXT("VTPackedUniform"));
-		PackedPageTableUniform0Parameter.Bind(ParameterMap, TEXT("VTPackedPageTableUniform0"));
-		PackedPageTableUniform1Parameter.Bind(ParameterMap, TEXT("VTPackedPageTableUniform1"));
-		PageTableSizeParameter.Bind(ParameterMap, TEXT("PageTableSize"));
-		PhysicalTextureSizeParameter.Bind(ParameterMap, TEXT("PhysicalTextureSize"));
-		MaxLodParameter.Bind(ParameterMap, TEXT("MaxLod"));
-		VirtualHeightfieldToLocalParameter.Bind(ParameterMap, TEXT("VirtualHeightfieldToLocal"));
-		VirtualHeightfieldToWorldParameter.Bind(ParameterMap, TEXT("VirtualHeightfieldToWorld"));
 		LodViewOriginParameter.Bind(ParameterMap, TEXT("LodViewOrigin"));
 		LodDistancesParameter.Bind(ParameterMap, TEXT("LodDistances"));
-		LodBiasScaleParameter.Bind(ParameterMap, TEXT("LodBiasScale"));
 	}
 
 	void GetElementShaderBindings(
@@ -130,94 +116,18 @@ public:
 		FVertexInputStreamArray& VertexStreams) const
 	{
 		FVirtualHeightfieldMeshVertexFactory* VertexFactory = (FVirtualHeightfieldMeshVertexFactory*)InVertexFactory;
-
-		ShaderBindings.Add(Shader->GetUniformBufferParameter<FVirtualHeightfieldMeshVertexFactoryParameters>(), VertexFactory->GetVirtualHeightfieldMeshVertexFactoryUniformBuffer());
+		ShaderBindings.Add(Shader->GetUniformBufferParameter<FVirtualHeightfieldMeshVertexFactoryParameters>(), VertexFactory->UniformBuffer);
 
 		FVirtualHeightfieldMeshUserData* UserData = (FVirtualHeightfieldMeshUserData*)BatchElement.UserData;
-		if (UserData != nullptr)
-		{
-			if (InstanceBufferParameter.IsBound())
-			{
-				ShaderBindings.Add(InstanceBufferParameter, UserData->InstanceBufferSRV);
-			}
-			if (PageTableTextureParameter.IsBound())
-			{
-				ShaderBindings.AddTexture(PageTableTextureParameter, PageTableSamplerParameter, TStaticSamplerState<SF_Point>::GetRHI(), UserData->PageTableTexture);
-			}
-			if (HeightTextureParameter.IsBound() && HeightSamplerParameter.IsBound())
-			{
-				ShaderBindings.AddTexture(HeightTextureParameter, HeightSamplerParameter, TStaticSamplerState<SF_Bilinear>::GetRHI(), UserData->HeightPhysicalTexture);
-			}
-			if (LodBiasTextureParameter.IsBound() && LodBiasSamplerParameter.IsBound())
-			{
-				ShaderBindings.AddTexture(LodBiasTextureParameter, LodBiasSamplerParameter, TStaticSamplerState<SF_Point>::GetRHI(), UserData->LodBiasTexture);
-			}
-			if (PackedUniformParameter.IsBound())
-			{
-				ShaderBindings.Add(PackedUniformParameter, UserData->PackedUniform);
-			}
-			if (PackedPageTableUniform0Parameter.IsBound())
-			{
-				ShaderBindings.Add(PackedPageTableUniform0Parameter, UserData->PackedPageTableUniform[0]);
-			}
-			if (PackedPageTableUniform1Parameter.IsBound())
-			{
-				ShaderBindings.Add(PackedPageTableUniform1Parameter, UserData->PackedPageTableUniform[1]);
-			}
-			if (PageTableSizeParameter.IsBound())
-			{
-				ShaderBindings.Add(PageTableSizeParameter, UserData->PageTableSize);
-			}
-			if (PhysicalTextureSizeParameter.IsBound())
-			{
-				ShaderBindings.Add(PhysicalTextureSizeParameter, UserData->PhysicalTextureSize);
-			}
-			if (MaxLodParameter.IsBound())
-			{
-				ShaderBindings.Add(MaxLodParameter, UserData->MaxLod);
-			}
-			if (VirtualHeightfieldToLocalParameter.IsBound())
-			{
-				ShaderBindings.Add(VirtualHeightfieldToLocalParameter, UserData->VirtualHeightfieldToLocal);
-			}
-			if (VirtualHeightfieldToWorldParameter.IsBound())
-			{
-				ShaderBindings.Add(VirtualHeightfieldToWorldParameter, UserData->VirtualHeightfieldToWorld);
-			}
-			if (LodViewOriginParameter.IsBound())
-			{
-				ShaderBindings.Add(LodViewOriginParameter, UserData->LodViewOrigin);
-			}
-			if (LodDistancesParameter.IsBound())
-			{
-				ShaderBindings.Add(LodDistancesParameter, UserData->LodDistances);
-			}
-			if (LodBiasScaleParameter.IsBound())
-			{
-				ShaderBindings.Add(LodBiasScaleParameter, UserData->LodBiasScale);
-			}
-		}
+		ShaderBindings.Add(InstanceBufferParameter, UserData->InstanceBufferSRV);
+		ShaderBindings.Add(LodViewOriginParameter, UserData->LodViewOrigin);
+		ShaderBindings.Add(LodDistancesParameter, UserData->LodDistances);
 	}
 
 protected:
 	LAYOUT_FIELD(FShaderResourceParameter, InstanceBufferParameter);
-	LAYOUT_FIELD(FShaderResourceParameter, PageTableTextureParameter);
-	LAYOUT_FIELD(FShaderResourceParameter, PageTableSamplerParameter);
-	LAYOUT_FIELD(FShaderResourceParameter, HeightTextureParameter);
-	LAYOUT_FIELD(FShaderResourceParameter, HeightSamplerParameter);
-	LAYOUT_FIELD(FShaderResourceParameter, LodBiasTextureParameter);
-	LAYOUT_FIELD(FShaderResourceParameter, LodBiasSamplerParameter);
-	LAYOUT_FIELD(FShaderParameter, PackedUniformParameter);
-	LAYOUT_FIELD(FShaderParameter, PackedPageTableUniform0Parameter);
-	LAYOUT_FIELD(FShaderParameter, PackedPageTableUniform1Parameter);
-	LAYOUT_FIELD(FShaderParameter, PageTableSizeParameter);
-	LAYOUT_FIELD(FShaderParameter, PhysicalTextureSizeParameter);
-	LAYOUT_FIELD(FShaderParameter, MaxLodParameter);
-	LAYOUT_FIELD(FShaderParameter, VirtualHeightfieldToLocalParameter);
-	LAYOUT_FIELD(FShaderParameter, VirtualHeightfieldToWorldParameter);
 	LAYOUT_FIELD(FShaderParameter, LodViewOriginParameter);
 	LAYOUT_FIELD(FShaderParameter, LodDistancesParameter);
-	LAYOUT_FIELD(FShaderParameter, LodBiasScaleParameter);
 };
 
 IMPLEMENT_TYPE_LAYOUT(FVirtualHeightfieldMeshVertexFactoryShaderParameters);
@@ -226,11 +136,11 @@ IMPLEMENT_VERTEX_FACTORY_PARAMETER_TYPE(FVirtualHeightfieldMeshVertexFactory, SF
 IMPLEMENT_VERTEX_FACTORY_PARAMETER_TYPE(FVirtualHeightfieldMeshVertexFactory, SF_Pixel, FVirtualHeightfieldMeshVertexFactoryShaderParameters);
 
 
-FVirtualHeightfieldMeshVertexFactory::FVirtualHeightfieldMeshVertexFactory(ERHIFeatureLevel::Type InFeatureLevel, int32 InNumQuadsPerSide)
+FVirtualHeightfieldMeshVertexFactory::FVirtualHeightfieldMeshVertexFactory(ERHIFeatureLevel::Type InFeatureLevel, const FVirtualHeightfieldMeshVertexFactoryParameters& InParams)
 	: FVertexFactory(InFeatureLevel)
-	, NumQuadsPerSide(InNumQuadsPerSide)
+	, Params(InParams)
 {
-	IndexBuffer = new FVirtualHeightfieldMeshIndexBuffer(NumQuadsPerSide);
+	IndexBuffer = new FVirtualHeightfieldMeshIndexBuffer(InParams.NumQuadsPerTileSide);
 }
 
 FVirtualHeightfieldMeshVertexFactory::~FVirtualHeightfieldMeshVertexFactory()
@@ -240,10 +150,7 @@ FVirtualHeightfieldMeshVertexFactory::~FVirtualHeightfieldMeshVertexFactory()
 
 void FVirtualHeightfieldMeshVertexFactory::InitRHI()
 {
-	// Setup the uniform data
-	FVirtualHeightfieldMeshVertexFactoryParameters UniformParams;
-	UniformParams.NumQuadsPerTileSide = NumQuadsPerSide;
-	UniformBuffer = FVirtualHeightfieldMeshVertexFactoryBufferRef::CreateUniformBufferImmediate(UniformParams, UniformBuffer_MultiFrame);
+	UniformBuffer = FVirtualHeightfieldMeshVertexFactoryBufferRef::CreateUniformBufferImmediate(Params, UniformBuffer_MultiFrame);
 
 	IndexBuffer->InitResource();
 
diff --git a/Engine/Plugins/Experimental/VirtualHeightfieldMesh/Source/VirtualHeightfieldMesh/Private/VirtualHeightfieldMeshVertexFactory.h b/Engine/Plugins/Experimental/VirtualHeightfieldMesh/Source/VirtualHeightfieldMesh/Private/VirtualHeightfieldMeshVertexFactory.h
--- a/Engine/Plugins/Experimental/VirtualHeightfieldMesh/Source/VirtualHeightfieldMesh/Private/VirtualHeightfieldMeshVertexFactory.h
+++ b/Engine/Plugins/Experimental/VirtualHeightfieldMesh/Source/VirtualHeightfieldMesh/Private/VirtualHeightfieldMeshVertexFactory.h
@@ -11,35 +11,40 @@
 #include "SceneManagement.h"
 
 /**
- * Uniform buffer to hold parameters specific to this vertex factory. Only set up once
+ * Uniform buffer to hold parameters specific to this vertex factory. Only set up once.
  */
 BEGIN_GLOBAL_SHADER_PARAMETER_STRUCT(FVirtualHeightfieldMeshVertexFactoryParameters, )
-	SHADER_PARAMETER(float, LODScale)
+	SHADER_PARAMETER_TEXTURE(Texture2D<uint4>, PageTableTexture)
+	SHADER_PARAMETER_SRV(Texture2D<float>, HeightTexture)
+	SHADER_PARAMETER_SAMPLER(SamplerState, HeightSampler)
+	SHADER_PARAMETER_TEXTURE(Texture2D<float>, LodBiasTexture)
+	SHADER_PARAMETER_SAMPLER(SamplerState, LodBiasSampler)
 	SHADER_PARAMETER(int32, NumQuadsPerTileSide)
+	SHADER_PARAMETER(FUintVector4, VTPackedUniform)
+	SHADER_PARAMETER(FUintVector4, VTPackedPageTableUniform0)
+	SHADER_PARAMETER(FUintVector4, VTPackedPageTableUniform1)
+	SHADER_PARAMETER(FVector4, PageTableSize)
+	SHADER_PARAMETER(FVector2D, PhysicalTextureSize)
+	SHADER_PARAMETER(FMatrix, VirtualHeightfieldToLocal)
+	SHADER_PARAMETER(FMatrix, VirtualHeightfieldToWorld)
+	SHADER_PARAMETER(float, MaxLod)
+	SHADER_PARAMETER(float, LodBiasScale)
 END_GLOBAL_SHADER_PARAMETER_STRUCT()
 
 typedef TUniformBufferRef<FVirtualHeightfieldMeshVertexFactoryParameters> FVirtualHeightfieldMeshVertexFactoryBufferRef;
 
+/**
+ * Per frame UserData to pass to the vertex shader.
+ */
 struct FVirtualHeightfieldMeshUserData : public FOneFrameResource
 {
 	FRHIShaderResourceView* InstanceBufferSRV;
-	FRHITexture* PageTableTexture;
-	FRHITexture* HeightPhysicalTexture;
-	FRHITexture* LodBiasTexture;
-	FUintVector4 PackedUniform;
-	FUintVector4 PackedPageTableUniform[2];
-	FVector4 PageTableSize;
-	FVector2D PhysicalTextureSize;
-	float MaxLod;
-	FMatrix VirtualHeightfieldToLocal;
-	FMatrix VirtualHeightfieldToWorld;
 	FVector LodViewOrigin;
 	FVector4 LodDistances;
-	float LodBiasScale;
 };
 
 /**
- * 
+ * Index buffer for a single Virtual Heightfield Mesh tile.
  */
 class FVirtualHeightfieldMeshIndexBuffer : public FIndexBuffer
 {
@@ -59,14 +64,14 @@ private:
 };
 
 /**
- *
+ * Virtual Heightfield Mesh vertex factory.
  */
 class FVirtualHeightfieldMeshVertexFactory : public FVertexFactory
 {
 	DECLARE_VERTEX_FACTORY_TYPE(FVirtualHeightfieldMeshVertexFactory);
 
 public:
-	FVirtualHeightfieldMeshVertexFactory(ERHIFeatureLevel::Type InFeatureLevel, int32 InNumQuadsPerSide);
+	FVirtualHeightfieldMeshVertexFactory(ERHIFeatureLevel::Type InFeatureLevel, const FVirtualHeightfieldMeshVertexFactoryParameters& InParams);
 
 	~FVirtualHeightfieldMeshVertexFactory();
 
@@ -77,15 +82,12 @@ public:
 	static void ModifyCompilationEnvironment(const FVertexFactoryShaderPermutationParameters& Parameters, FShaderCompilerEnvironment& OutEnvironment);
 	static void ValidateCompiledResult(const FVertexFactoryType* Type, EShaderPlatform Platform, const FShaderParameterMap& ParameterMap, TArray<FString>& OutErrors);
 
-	inline const FUniformBufferRHIRef GetVirtualHeightfieldMeshVertexFactoryUniformBuffer() const
-	{
-		return UniformBuffer;
-	}
-
-	FVirtualHeightfieldMeshIndexBuffer* IndexBuffer = nullptr;
+	FIndexBuffer const* GetIndexBuffer() const { return IndexBuffer; }
 
 private:
+	FVirtualHeightfieldMeshVertexFactoryParameters Params;
 	FVirtualHeightfieldMeshVertexFactoryBufferRef UniformBuffer;
+	FVirtualHeightfieldMeshIndexBuffer* IndexBuffer = nullptr;
 
-	const int32 NumQuadsPerSide = 0;
+	friend class FVirtualHeightfieldMeshVertexFactoryShaderParameters;
 };
