Removes template based shader permutation of FCullInstancesCS

diff --git a/Engine/Plugins/Experimental/VirtualHeightfieldMesh/Source/VirtualHeightfieldMesh/Private/VirtualHeightfieldMeshSceneProxy.cpp b/Engine/Plugins/Experimental/VirtualHeightfieldMesh/Source/VirtualHeightfieldMesh/Private/VirtualHeightfieldMeshSceneProxy.cpp
--- a/Engine/Plugins/Experimental/VirtualHeightfieldMesh/Source/VirtualHeightfieldMesh/Private/VirtualHeightfieldMeshSceneProxy.cpp
+++ b/Engine/Plugins/Experimental/VirtualHeightfieldMesh/Source/VirtualHeightfieldMesh/Private/VirtualHeightfieldMeshSceneProxy.cpp
@@ -794,10 +794,20 @@ namespace VirtualHeightfieldMesh
 	IMPLEMENT_GLOBAL_SHADER(FInitInstanceBufferCS, "/Plugin/VirtualHeightfieldMesh/Private/VirtualHeightfieldMesh.usf", "InitInstanceBufferCS", SF_Compute);
 
 	/** CullInstances compute shader. */
-	class FCullInstances : public FGlobalShader
+	class FCullInstancesCS : public FGlobalShader
 	{
 	public:
-		SHADER_USE_PARAMETER_STRUCT(FCullInstances, FGlobalShader);
+		DECLARE_GLOBAL_SHADER(FCullInstancesCS);
+		SHADER_USE_PARAMETER_STRUCT(FCullInstancesCS, FGlobalShader);
+
+		class FReuseCullDim : SHADER_PERMUTATION_BOOL("REUSE_CULL");
+
+		using FPermutationDomain = TShaderPermutationDomain<FReuseCullDim>;
+
+		static bool ShouldCompilePermutation(FGlobalShaderPermutationParameters const& Parameters)
+		{
+			return IsFeatureLevelSupported(Parameters.Platform, ERHIFeatureLevel::SM5);
+		}
 
 		BEGIN_SHADER_PARAMETER_STRUCT(FParameters, )
 			SHADER_PARAMETER_TEXTURE(Texture2D, HeightMinMaxTexture)
@@ -816,34 +826,7 @@ namespace VirtualHeightfieldMesh
 		END_SHADER_PARAMETER_STRUCT()
 	};
 
-	template< bool bReuseCull >
-	class TCullInstancesCS : public FCullInstances
-	{
-	public:
-		typedef TCullInstancesCS< bReuseCull > ClassName;
-		DECLARE_GLOBAL_SHADER(ClassName);
-
-		TCullInstancesCS()
-		{}
-
-		TCullInstancesCS(const ShaderMetaType::CompiledShaderInitializerType& Initializer)
-			: FCullInstances(Initializer)
-		{}
-
-		static bool ShouldCompilePermutation(FGlobalShaderPermutationParameters const& Parameters)
-		{
-			return IsFeatureLevelSupported(Parameters.Platform, ERHIFeatureLevel::SM5);
-		}
-
-		static void ModifyCompilationEnvironment(const FGlobalShaderPermutationParameters& Parameters, FShaderCompilerEnvironment& OutEnvironment)
-		{
-			FGlobalShader::ModifyCompilationEnvironment(Parameters, OutEnvironment);
-			OutEnvironment.SetDefine(TEXT("REUSE_CULL"), bReuseCull ? 1 : 0);
-		}
-	};
-
-	IMPLEMENT_SHADER_TYPE(template<>, TCullInstancesCS<true>, TEXT("/Plugin/VirtualHeightfieldMesh/Private/VirtualHeightfieldMesh.usf"), TEXT("CullInstancesCS"), SF_Compute);
-	IMPLEMENT_SHADER_TYPE(template<>, TCullInstancesCS<false>, TEXT("/Plugin/VirtualHeightfieldMesh/Private/VirtualHeightfieldMesh.usf"), TEXT("CullInstancesCS"), SF_Compute);
+	IMPLEMENT_GLOBAL_SHADER(FCullInstancesCS, "/Plugin/VirtualHeightfieldMesh/Private/VirtualHeightfieldMesh.usf", "CullInstancesCS", SF_Compute);
 
 
 	/** Default Min/Max texture has the fixed maximum [0,1]. */
@@ -1155,7 +1138,7 @@ namespace VirtualHeightfieldMesh
 	/** Cull quads and write to the final output buffer. */
 	void AddPass_CullInstances(FRDGBuilder& GraphBuilder, FGlobalShaderMap* InGlobalShaderMap, FProxyDesc const& InDesc, FVolatileResources& InVolatileResources, FDrawInstanceBuffers& InOutputResources, FChildViewDesc const& InViewDesc)
 	{
-		FCullInstances::FParameters* PassParameters = GraphBuilder.AllocParameters<FCullInstances::FParameters>();
+		FCullInstancesCS::FParameters* PassParameters = GraphBuilder.AllocParameters<FCullInstancesCS::FParameters>();
 		PassParameters->HeightMinMaxTexture = InDesc.HeightMinMaxTexture;
 		PassParameters->MinMaxTextureSampler = TStaticSamplerState<SF_Point>::GetRHI();
 		PassParameters->MinMaxLevelOffset = InDesc.MinMaxLevelOffset;
@@ -1173,26 +1156,18 @@ namespace VirtualHeightfieldMesh
 		PassParameters->RWInstanceBuffer = InOutputResources.InstanceBufferUAV;
 		PassParameters->RWIndirectArgsBuffer = InOutputResources.IndirectArgsBufferUAV;
 
-		TShaderRef<FCullInstances> ComputeShader;
-		if (InViewDesc.bIsMainView)
-		{
-			ComputeShader = InGlobalShaderMap->GetShader< TCullInstancesCS<true> >();
-		}
-		else
-		{
-			ComputeShader = InGlobalShaderMap->GetShader< TCullInstancesCS<false> >();
-		}
+		int32 IndirectArgOffset = VirtualHeightfieldMesh::IndirectArgsByteOffset_FinalCull;
 
-		GraphBuilder.AddPass(
+		FCullInstancesCS::FPermutationDomain PermutationVector;
+		PermutationVector.Set<FCullInstancesCS::FReuseCullDim>(InViewDesc.bIsMainView);
+
+		TShaderMapRef<FCullInstancesCS> ComputeShader(InGlobalShaderMap, PermutationVector);
+		FComputeShaderUtils::AddPass(
+			GraphBuilder,
 			RDG_EVENT_NAME("CullInstances"),
-			PassParameters,
-			ERDGPassFlags::Compute,
-			[PassParameters, ComputeShader, IndirectArgsBuffer = InVolatileResources.IndirectArgsBuffer](FRHICommandList& RHICmdList)
-		{
-			int32 IndirectArgOffset = VirtualHeightfieldMesh::IndirectArgsByteOffset_FinalCull;
-			FComputeShaderUtils::DispatchIndirect(RHICmdList, ComputeShader, *PassParameters, IndirectArgsBuffer->GetIndirectRHICallBuffer(), IndirectArgOffset);
-			IndirectArgsBuffer->MarkResourceAsUsed();
-		});
+			ComputeShader, PassParameters,
+			InVolatileResources.IndirectArgsBuffer,
+			IndirectArgOffset);
 	}
 }
 
