Fix global shader naming collision introduced by previous change

diff --git a/Engine/Plugins/Experimental/VirtualHeightfieldMesh/Source/VirtualHeightfieldMesh/Private/VirtualHeightfieldMeshSceneProxy.cpp b/Engine/Plugins/Experimental/VirtualHeightfieldMesh/Source/VirtualHeightfieldMesh/Private/VirtualHeightfieldMeshSceneProxy.cpp
--- a/Engine/Plugins/Experimental/VirtualHeightfieldMesh/Source/VirtualHeightfieldMesh/Private/VirtualHeightfieldMeshSceneProxy.cpp
+++ b/Engine/Plugins/Experimental/VirtualHeightfieldMesh/Source/VirtualHeightfieldMesh/Private/VirtualHeightfieldMeshSceneProxy.cpp
@@ -708,11 +708,11 @@ namespace VirtualHeightfieldMesh
 	};
 
 	/** Compute shader to initialize all buffers, including adding the lowest mip page(s) to the QuadBuffer. */
-	class FInitBuffersCS : public FGlobalShader
+	class FInitBuffersVHM_CS : public FGlobalShader
 	{
 	public:
-		DECLARE_GLOBAL_SHADER(FInitBuffersCS);
-		SHADER_USE_PARAMETER_STRUCT(FInitBuffersCS, FGlobalShader);
+		DECLARE_GLOBAL_SHADER(FInitBuffersVHM_CS);
+		SHADER_USE_PARAMETER_STRUCT(FInitBuffersVHM_CS, FGlobalShader);
 
 		BEGIN_SHADER_PARAMETER_STRUCT(FParameters, )
 			SHADER_PARAMETER(uint32, MaxLevel)
@@ -731,14 +731,14 @@ namespace VirtualHeightfieldMesh
 		}
 	};
 
-	IMPLEMENT_GLOBAL_SHADER(FInitBuffersCS, "/Plugin/VirtualHeightfieldMesh/Private/VirtualHeightfieldMesh.usf", "InitBuffersCS", SF_Compute);
+	IMPLEMENT_GLOBAL_SHADER(FInitBuffersVHM_CS, "/Plugin/VirtualHeightfieldMesh/Private/VirtualHeightfieldMesh.usf", "InitBuffersCS", SF_Compute);
 
 	/** Compute shader to traverse the virtual texture page table for a view and generate an array of quads to potentially render. */
-	class FCollectQuadsCS : public FGlobalShader
+	class FCollectQuadsVHM_CS : public FGlobalShader
 	{
 	public:
-		DECLARE_GLOBAL_SHADER(FCollectQuadsCS);
-		SHADER_USE_PARAMETER_STRUCT(FCollectQuadsCS, FGlobalShader);
+		DECLARE_GLOBAL_SHADER(FCollectQuadsVHM_CS);
+		SHADER_USE_PARAMETER_STRUCT(FCollectQuadsVHM_CS, FGlobalShader);
 
 		BEGIN_SHADER_PARAMETER_STRUCT(FParameters, )
 			SHADER_PARAMETER_TEXTURE(Texture2D, HeightMinMaxTexture)
@@ -771,14 +771,14 @@ namespace VirtualHeightfieldMesh
 		}
 	};
 
-	IMPLEMENT_GLOBAL_SHADER(FCollectQuadsCS, "/Plugin/VirtualHeightfieldMesh/Private/VirtualHeightfieldMesh.usf", "CollectQuadsCS", SF_Compute);
+	IMPLEMENT_GLOBAL_SHADER(FCollectQuadsVHM_CS, "/Plugin/VirtualHeightfieldMesh/Private/VirtualHeightfieldMesh.usf", "CollectQuadsCS", SF_Compute);
 
 	/** InitInstanceBuffer compute shader. */
-	class FInitInstanceBufferCS : public FGlobalShader
+	class FInitInstanceBufferVHM_CS : public FGlobalShader
 	{
 	public:
-		DECLARE_GLOBAL_SHADER(FInitInstanceBufferCS);
-		SHADER_USE_PARAMETER_STRUCT(FInitInstanceBufferCS, FGlobalShader);
+		DECLARE_GLOBAL_SHADER(FInitInstanceBufferVHM_CS);
+		SHADER_USE_PARAMETER_STRUCT(FInitInstanceBufferVHM_CS, FGlobalShader);
 
 		BEGIN_SHADER_PARAMETER_STRUCT(FParameters, )
 			SHADER_PARAMETER(int32, NumIndices)
@@ -791,14 +791,14 @@ namespace VirtualHeightfieldMesh
 		}
 	};
 
-	IMPLEMENT_GLOBAL_SHADER(FInitInstanceBufferCS, "/Plugin/VirtualHeightfieldMesh/Private/VirtualHeightfieldMesh.usf", "InitInstanceBufferCS", SF_Compute);
+	IMPLEMENT_GLOBAL_SHADER(FInitInstanceBufferVHM_CS, "/Plugin/VirtualHeightfieldMesh/Private/VirtualHeightfieldMesh.usf", "InitInstanceBufferCS", SF_Compute);
 
 	/** CullInstances compute shader. */
-	class FCullInstancesCS : public FGlobalShader
+	class FCullInstancesVHM_CS : public FGlobalShader
 	{
 	public:
-		DECLARE_GLOBAL_SHADER(FCullInstancesCS);
-		SHADER_USE_PARAMETER_STRUCT(FCullInstancesCS, FGlobalShader);
+		DECLARE_GLOBAL_SHADER(FCullInstancesVHM_CS);
+		SHADER_USE_PARAMETER_STRUCT(FCullInstancesVHM_CS, FGlobalShader);
 
 		class FReuseCullDim : SHADER_PERMUTATION_BOOL("REUSE_CULL");
 
@@ -826,7 +826,7 @@ namespace VirtualHeightfieldMesh
 		END_SHADER_PARAMETER_STRUCT()
 	};
 
-	IMPLEMENT_GLOBAL_SHADER(FCullInstancesCS, "/Plugin/VirtualHeightfieldMesh/Private/VirtualHeightfieldMesh.usf", "CullInstancesCS", SF_Compute);
+	IMPLEMENT_GLOBAL_SHADER(FCullInstancesVHM_CS, "/Plugin/VirtualHeightfieldMesh/Private/VirtualHeightfieldMesh.usf", "CullInstancesCS", SF_Compute);
 
 
 	/** Default Min/Max texture has the fixed maximum [0,1]. */
@@ -1057,9 +1057,9 @@ namespace VirtualHeightfieldMesh
 	/** Initialize the buffers before collecting visible quads. */
 	void AddPass_InitBuffers(FRDGBuilder& GraphBuilder, FGlobalShaderMap* InGlobalShaderMap, FProxyDesc const& InDesc, FVolatileResources& InVolatileResources)
 	{
-		TShaderMapRef<FInitBuffersCS> ComputeShader(InGlobalShaderMap);
+		TShaderMapRef<FInitBuffersVHM_CS> ComputeShader(InGlobalShaderMap);
 
-		FInitBuffersCS::FParameters* PassParameters = GraphBuilder.AllocParameters<FInitBuffersCS::FParameters>();
+		FInitBuffersVHM_CS::FParameters* PassParameters = GraphBuilder.AllocParameters<FInitBuffersVHM_CS::FParameters>();
 		PassParameters->MaxLevel = InDesc.MaxLevel;
 		PassParameters->NumForceLoadLods = InDesc.NumForceLoadLods;
 		PassParameters->PageTableFeedbackId = InDesc.PageTableFeedbackId;
@@ -1085,9 +1085,9 @@ namespace VirtualHeightfieldMesh
 	/** Collect potentially visible quads and determine their Lods. */
 	void AddPass_CollectQuads(FRDGBuilder& GraphBuilder, FGlobalShaderMap* InGlobalShaderMap, FProxyDesc const& InDesc, FVolatileResources& InVolatileResources, FMainViewDesc const& InViewDesc)
 	{
-		TShaderMapRef<FCollectQuadsCS> ComputeShader(InGlobalShaderMap);
+		TShaderMapRef<FCollectQuadsVHM_CS> ComputeShader(InGlobalShaderMap);
 
-		FCollectQuadsCS::FParameters* PassParameters = GraphBuilder.AllocParameters<FCollectQuadsCS::FParameters>();
+		FCollectQuadsVHM_CS::FParameters* PassParameters = GraphBuilder.AllocParameters<FCollectQuadsVHM_CS::FParameters>();
 		PassParameters->HeightMinMaxTexture = InDesc.HeightMinMaxTexture;
 		PassParameters->LodBiasMinMaxTexture = InDesc.LodBiasMinMaxTexture;
 		PassParameters->MinMaxTextureSampler = TStaticSamplerState<SF_Point>::GetRHI();
@@ -1123,9 +1123,9 @@ namespace VirtualHeightfieldMesh
 	/** Initialise the draw indirect buffer. */
 	void AddPass_InitInstanceBuffer(FRDGBuilder& GraphBuilder, FGlobalShaderMap* InGlobalShaderMap, int32 NumQuadsPerTileSide, FDrawInstanceBuffers& InOutputResources)
 	{
-		TShaderMapRef<FInitInstanceBufferCS> ComputeShader(InGlobalShaderMap);
+		TShaderMapRef<FInitInstanceBufferVHM_CS> ComputeShader(InGlobalShaderMap);
 
-		FInitInstanceBufferCS::FParameters* PassParameters = GraphBuilder.AllocParameters<FInitInstanceBufferCS::FParameters>();
+		FInitInstanceBufferVHM_CS::FParameters* PassParameters = GraphBuilder.AllocParameters<FInitInstanceBufferVHM_CS::FParameters>();
 		PassParameters->NumIndices = NumQuadsPerTileSide * NumQuadsPerTileSide * 6;
 		PassParameters->RWIndirectArgsBuffer = InOutputResources.IndirectArgsBufferUAV;
 
@@ -1138,7 +1138,7 @@ namespace VirtualHeightfieldMesh
 	/** Cull quads and write to the final output buffer. */
 	void AddPass_CullInstances(FRDGBuilder& GraphBuilder, FGlobalShaderMap* InGlobalShaderMap, FProxyDesc const& InDesc, FVolatileResources& InVolatileResources, FDrawInstanceBuffers& InOutputResources, FChildViewDesc const& InViewDesc)
 	{
-		FCullInstancesCS::FParameters* PassParameters = GraphBuilder.AllocParameters<FCullInstancesCS::FParameters>();
+		FCullInstancesVHM_CS::FParameters* PassParameters = GraphBuilder.AllocParameters<FCullInstancesVHM_CS::FParameters>();
 		PassParameters->HeightMinMaxTexture = InDesc.HeightMinMaxTexture;
 		PassParameters->MinMaxTextureSampler = TStaticSamplerState<SF_Point>::GetRHI();
 		PassParameters->MinMaxLevelOffset = InDesc.MinMaxLevelOffset;
@@ -1158,10 +1158,10 @@ namespace VirtualHeightfieldMesh
 
 		int32 IndirectArgOffset = VirtualHeightfieldMesh::IndirectArgsByteOffset_FinalCull;
 
-		FCullInstancesCS::FPermutationDomain PermutationVector;
-		PermutationVector.Set<FCullInstancesCS::FReuseCullDim>(InViewDesc.bIsMainView);
+		FCullInstancesVHM_CS::FPermutationDomain PermutationVector;
+		PermutationVector.Set<FCullInstancesVHM_CS::FReuseCullDim>(InViewDesc.bIsMainView);
 
-		TShaderMapRef<FCullInstancesCS> ComputeShader(InGlobalShaderMap, PermutationVector);
+		TShaderMapRef<FCullInstancesVHM_CS> ComputeShader(InGlobalShaderMap, PermutationVector);
 		FComputeShaderUtils::AddPass(
 			GraphBuilder,
 			RDG_EVENT_NAME("CullInstances"),
