VHM: Fix for shader bind error with Large World Coords.
In future needs adjusting to deal with large values.

diff --git a/Engine/Plugins/Experimental/VirtualHeightfieldMesh/Source/VirtualHeightfieldMesh/Private/VirtualHeightfieldMeshVertexFactory.h b/Engine/Plugins/Experimental/VirtualHeightfieldMesh/Source/VirtualHeightfieldMesh/Private/VirtualHeightfieldMeshVertexFactory.h
--- a/Engine/Plugins/Experimental/VirtualHeightfieldMesh/Source/VirtualHeightfieldMesh/Private/VirtualHeightfieldMeshVertexFactory.h
+++ b/Engine/Plugins/Experimental/VirtualHeightfieldMesh/Source/VirtualHeightfieldMesh/Private/VirtualHeightfieldMeshVertexFactory.h
@@ -39,7 +39,7 @@ typedef TUniformBufferRef<FVirtualHeightfieldMeshVertexFactoryParameters> FVirtu
 struct FVirtualHeightfieldMeshUserData : public FOneFrameResource
 {
 	FRHIShaderResourceView* InstanceBufferSRV;
-	FVector LodViewOrigin;
+	FVector3f LodViewOrigin;
 	FVector4 LodDistances;
 };
 
