VHM: Add SetMaterial support

diff --git a/Engine/Plugins/Experimental/VirtualHeightfieldMesh/Source/VirtualHeightfieldMesh/Private/VirtualHeightfieldMeshComponent.cpp b/Engine/Plugins/Experimental/VirtualHeightfieldMesh/Source/VirtualHeightfieldMesh/Private/VirtualHeightfieldMeshComponent.cpp
--- a/Engine/Plugins/Experimental/VirtualHeightfieldMesh/Source/VirtualHeightfieldMesh/Private/VirtualHeightfieldMeshComponent.cpp
+++ b/Engine/Plugins/Experimental/VirtualHeightfieldMesh/Source/VirtualHeightfieldMesh/Private/VirtualHeightfieldMeshComponent.cpp
@@ -95,6 +95,15 @@ FPrimitiveSceneProxy* UVirtualHeightfieldMeshComponent::CreateSceneProxy()
 	return bIsEnabled ? new FVirtualHeightfieldMeshSceneProxy(this) : nullptr;
 }
 
+void UVirtualHeightfieldMeshComponent::SetMaterial(int32 InElementIndex, class UMaterialInterface* InMaterial)
+{
+	if (ElementIndex == 0 && Material != InMaterial)
+	{
+		Material = InMaterial;
+		MarkRenderStateDirty();
+	}
+}
+
 void UVirtualHeightfieldMeshComponent::GetUsedMaterials(TArray<UMaterialInterface*>& OutMaterials, bool bGetDebugMaterials) const
 {
 	if (Material != nullptr)
diff --git a/Engine/Plugins/Experimental/VirtualHeightfieldMesh/Source/VirtualHeightfieldMesh/Public/VirtualHeightfieldMeshComponent.h b/Engine/Plugins/Experimental/VirtualHeightfieldMesh/Source/VirtualHeightfieldMesh/Public/VirtualHeightfieldMeshComponent.h
--- a/Engine/Plugins/Experimental/VirtualHeightfieldMesh/Source/VirtualHeightfieldMesh/Public/VirtualHeightfieldMeshComponent.h
+++ b/Engine/Plugins/Experimental/VirtualHeightfieldMesh/Source/VirtualHeightfieldMesh/Public/VirtualHeightfieldMeshComponent.h
@@ -138,6 +138,7 @@ protected:
 	//~ Begin UPrimitiveComponent Interface
 	virtual FPrimitiveSceneProxy* CreateSceneProxy() override;
 	virtual bool SupportsStaticLighting() const override { return true; }
+	virtual void SetMaterial(int32 ElementIndex, class UMaterialInterface* Material) override;
 	virtual UMaterialInterface* GetMaterial(int32 Index) const override { return Material; }
 	virtual void GetUsedMaterials(TArray<UMaterialInterface*>& OutMaterials, bool bGetDebugMaterials = false) const override;
 	//~ End UPrimitiveComponent Interface
