VHM: Build fix

diff --git a/Engine/Plugins/Experimental/VirtualHeightfieldMesh/Source/VirtualHeightfieldMesh/Private/VirtualHeightfieldMeshComponent.cpp b/Engine/Plugins/Experimental/VirtualHeightfieldMesh/Source/VirtualHeightfieldMesh/Private/VirtualHeightfieldMeshComponent.cpp
--- a/Engine/Plugins/Experimental/VirtualHeightfieldMesh/Source/VirtualHeightfieldMesh/Private/VirtualHeightfieldMeshComponent.cpp
+++ b/Engine/Plugins/Experimental/VirtualHeightfieldMesh/Source/VirtualHeightfieldMesh/Private/VirtualHeightfieldMeshComponent.cpp
@@ -95,9 +95,9 @@ FPrimitiveSceneProxy* UVirtualHeightfieldMeshComponent::CreateSceneProxy()
 	return bIsEnabled ? new FVirtualHeightfieldMeshSceneProxy(this) : nullptr;
 }
 
-void UVirtualHeightfieldMeshComponent::SetMaterial(int32 InElementIndex, class UMaterialInterface* InMaterial)
+void UVirtualHeightfieldMeshComponent::SetMaterial(int32 InElementIndex, UMaterialInterface* InMaterial)
 {
-	if (ElementIndex == 0 && Material != InMaterial)
+	if (InElementIndex == 0 && Material != InMaterial)
 	{
 		Material = InMaterial;
 		MarkRenderStateDirty();
