PR #8142
Fix disappearing VirtualHeightfieldMesh used with Translucent/SingleLayerWater material type.
(Contributed by [REDACTED_PERSON])

diff --git a/Engine/Plugins/Experimental/VirtualHeightfieldMesh/Source/VirtualHeightfieldMesh/Private/VirtualHeightfieldMeshSceneProxy.cpp b/Engine/Plugins/Experimental/VirtualHeightfieldMesh/Source/VirtualHeightfieldMesh/Private/VirtualHeightfieldMeshSceneProxy.cpp
--- a/Engine/Plugins/Experimental/VirtualHeightfieldMesh/Source/VirtualHeightfieldMesh/Private/VirtualHeightfieldMeshSceneProxy.cpp
+++ b/Engine/Plugins/Experimental/VirtualHeightfieldMesh/Source/VirtualHeightfieldMesh/Private/VirtualHeightfieldMeshSceneProxy.cpp
@@ -381,6 +381,7 @@ FVirtualHeightfieldMeshSceneProxy::FVirtualHeightfieldMeshSceneProxy(UVirtualHei
 	UMaterialInterface* ComponentMaterial = InComponent->GetMaterial();
 	const bool bValidMaterial = ComponentMaterial != nullptr && ComponentMaterial->CheckMaterialUsage_Concurrent(MATUSAGE_VirtualHeightfieldMesh);
 	Material = bValidMaterial ? ComponentMaterial->GetRenderProxy() : UMaterial::GetDefaultMaterial(MD_Surface)->GetRenderProxy();
+	MaterialRelevance = Material->GetMaterialInterface()->GetRelevance_Concurrent(GetScene().GetFeatureLevel());
 
 	const FTransform VirtualTextureTransform = InComponent->GetVirtualTextureTransform();
 
@@ -517,6 +518,7 @@ FPrimitiveViewRelevance FVirtualHeightfieldMeshSceneProxy::GetViewRelevance(cons
 	Result.bRenderCustomDepth = ShouldRenderCustomDepth();
 	Result.bTranslucentSelfShadow = false;
 	Result.bVelocityRelevance = false;
+	MaterialRelevance.SetPrimitiveViewRelevance(Result);
 	return Result;
 }
 
diff --git a/Engine/Plugins/Experimental/VirtualHeightfieldMesh/Source/VirtualHeightfieldMesh/Private/VirtualHeightfieldMeshSceneProxy.h b/Engine/Plugins/Experimental/VirtualHeightfieldMesh/Source/VirtualHeightfieldMesh/Private/VirtualHeightfieldMeshSceneProxy.h
--- a/Engine/Plugins/Experimental/VirtualHeightfieldMesh/Source/VirtualHeightfieldMesh/Private/VirtualHeightfieldMeshSceneProxy.h
+++ b/Engine/Plugins/Experimental/VirtualHeightfieldMesh/Source/VirtualHeightfieldMesh/Private/VirtualHeightfieldMeshSceneProxy.h
@@ -34,7 +34,8 @@ public:
 
 	class URuntimeVirtualTexture* RuntimeVirtualTexture;
 	class FMaterialRenderProxy* Material;
-	
+	FMaterialRelevance MaterialRelevance;
+
 	class UTexture2D* HeightMinMaxTexture;
 	class UTexture2D* LodBiasTexture;
 	class UTexture2D* LodBiasMinMaxTexture;
