VirtualHeightfieldMesh: Fix GPU hangs.
This fix needs to disable VHM shadows which will be reenabled in 5.1.

diff --git a/Engine/Plugins/Experimental/VirtualHeightfieldMesh/Source/VirtualHeightfieldMesh/Private/VirtualHeightfieldMeshSceneProxy.cpp b/Engine/Plugins/Experimental/VirtualHeightfieldMesh/Source/VirtualHeightfieldMesh/Private/VirtualHeightfieldMeshSceneProxy.cpp
--- a/Engine/Plugins/Experimental/VirtualHeightfieldMesh/Source/VirtualHeightfieldMesh/Private/VirtualHeightfieldMeshSceneProxy.cpp
+++ b/Engine/Plugins/Experimental/VirtualHeightfieldMesh/Source/VirtualHeightfieldMesh/Private/VirtualHeightfieldMeshSceneProxy.cpp
@@ -173,6 +173,9 @@ public:
 	/** Call once to register this extension. */
 	void RegisterExtension();
 
+	/** Are we inside a BeginFrame()/EndFrame() scope? */
+	bool IsInFrame() { return bInFrame; }
+
 	/** Call once per frame for each mesh/view that has relevance. This allocates the buffers to use for the frame and adds the work to fill the buffers to the queue. */
 	VirtualHeightfieldMesh::FDrawInstanceBuffers& AddWork(FVirtualHeightfieldMeshSceneProxy const* InProxy, FSceneView const* InMainView, FSceneView const* InCullView);
 	/** Submit all the work added by AddWork(). The work fills all of the buffers ready for use by the referencing mesh batches. */
@@ -530,6 +533,16 @@ void FVirtualHeightfieldMeshSceneProxy::GetDynamicMeshElements(const TArray<cons
 	check(IsInRenderingThread());
 	check(AllocatedVirtualTexture != nullptr);
 
+	if (GVirtualHeightfieldMeshViewRendererExtension.IsInFrame())
+	{
+		// Can't add new work while bInFrame.
+		// In UE5 we need to AddWork()/SubmitWork() in two phases: InitViews() and InitViewsAfterPrepass()
+		// The main renderer hooks for that don't exist in UE5.0 and are only added in UE5.1
+		// That means that for UE5.0 we always hit this for shadow drawing and shadows will not be rendered.
+		// Not earlying out here can lead to crashes from buffers being released too soon.
+		return;
+	}
+
 	for (int32 ViewIndex = 0; ViewIndex < Views.Num(); ViewIndex++)
 	{
 		if (VisibilityMap & (1 << ViewIndex))
