VHM: Fix reading uninitialized LocalToWorld matrix inFVirtualHeightfieldMeshSceneProxy constructor.
It's only safe to init UVToLocal in OnTransformChanged() which is always called immediately after construction.

diff --git a/Engine/Plugins/Experimental/VirtualHeightfieldMesh/Source/VirtualHeightfieldMesh/Private/VirtualHeightfieldMeshSceneProxy.cpp b/Engine/Plugins/Experimental/VirtualHeightfieldMesh/Source/VirtualHeightfieldMesh/Private/VirtualHeightfieldMeshSceneProxy.cpp
--- a/Engine/Plugins/Experimental/VirtualHeightfieldMesh/Source/VirtualHeightfieldMesh/Private/VirtualHeightfieldMeshSceneProxy.cpp
+++ b/Engine/Plugins/Experimental/VirtualHeightfieldMesh/Source/VirtualHeightfieldMesh/Private/VirtualHeightfieldMeshSceneProxy.cpp
@@ -397,7 +397,8 @@ FVirtualHeightfieldMeshSceneProxy::FVirtualHeightfieldMeshSceneProxy(UVirtualHei
 	WorldToUV = UVToWorld.Inverse();
 	WorldToUVTransposeAdjoint = WorldToUV.TransposeAdjoint();
 
-	UVToLocal = UVToWorld * GetLocalToWorld().Inverse();
+	// UVToLocal will be initialized in OnTransformChanged() called immediately after construction.
+	UVToLocal = FMatrix::Identity;
 
 	UHeightfieldMinMaxTexture* HeightfieldMinMaxTexture = InComponent->GetMinMaxTexture();
 	if (HeightfieldMinMaxTexture != nullptr)
