VirtualHeightfieldMesh: Fix issues when rendering more than one VHM.
licensee-reported.

diff --git a/Engine/Plugins/Experimental/VirtualHeightfieldMesh/Source/VirtualHeightfieldMesh/Private/VirtualHeightfieldMeshSceneProxy.cpp b/Engine/Plugins/Experimental/VirtualHeightfieldMesh/Source/VirtualHeightfieldMesh/Private/VirtualHeightfieldMeshSceneProxy.cpp
--- a/Engine/Plugins/Experimental/VirtualHeightfieldMesh/Source/VirtualHeightfieldMesh/Private/VirtualHeightfieldMeshSceneProxy.cpp
+++ b/Engine/Plugins/Experimental/VirtualHeightfieldMesh/Source/VirtualHeightfieldMesh/Private/VirtualHeightfieldMeshSceneProxy.cpp
@@ -1312,7 +1312,7 @@ void FVirtualHeightfieldMeshRendererExtension::SubmitWork(FRDGBuilder& GraphBuil
  			Desc.Init(CVarVHMMaxFeedbackItems.GetValueOnRenderThread() + 1);
  			SubmitVirtualTextureFeedbackBuffer(GraphBuilder, VolatileResources.FeedbackBuffer, Desc);
 
-			while (WorkIndex < NumWorkItems && MainViews[WorkDescs[WorkIndex].MainViewIndex] == MainView)
+			while (WorkIndex < NumWorkItems && SceneProxies[WorkDescs[WorkIndex].ProxyIndex] == Proxy && MainViews[WorkDescs[WorkIndex].MainViewIndex] == MainView)
 			{
 				// Gather data per child view
 				FSceneView const* CullView = CullViews[WorkDescs[WorkIndex].CullViewIndex];
