Fix issue where VirtualHeightfieldMesh with WorldPositionOffset fails to render correctly when velocity is in depth pass.

diff --git a/Engine/Plugins/Experimental/VirtualHeightfieldMesh/Source/VirtualHeightfieldMesh/Private/VirtualHeightfieldMeshSceneProxy.cpp b/Engine/Plugins/Experimental/VirtualHeightfieldMesh/Source/VirtualHeightfieldMesh/Private/VirtualHeightfieldMeshSceneProxy.cpp
--- a/Engine/Plugins/Experimental/VirtualHeightfieldMesh/Source/VirtualHeightfieldMesh/Private/VirtualHeightfieldMeshSceneProxy.cpp
+++ b/Engine/Plugins/Experimental/VirtualHeightfieldMesh/Source/VirtualHeightfieldMesh/Private/VirtualHeightfieldMeshSceneProxy.cpp
@@ -527,8 +527,8 @@ FPrimitiveViewRelevance FVirtualHeightfieldMeshSceneProxy::GetViewRelevance(cons
 	Result.bUsesLightingChannels = GetLightingChannelMask() != GetDefaultLightingChannelMask();
 	Result.bRenderCustomDepth = ShouldRenderCustomDepth();
 	Result.bTranslucentSelfShadow = false;
-	Result.bVelocityRelevance = false;
 	MaterialRelevance.SetPrimitiveViewRelevance(Result);
+	Result.bVelocityRelevance = DrawsVelocity() && Result.bOpaque && Result.bRenderInMainPass;
 	return Result;
 }
 
