VHM: Add component flag to locally set whether WPO materials should trigger velocity rendering.
Default off, since the anticipated usage of any WPO on a VHM is for displacement, rather than moving surface.
This then avoids unecessary cost of velocity calculation on the VHM.

diff --git a/Engine/Plugins/Experimental/VirtualHeightfieldMesh/Source/VirtualHeightfieldMesh/Public/VirtualHeightfieldMeshComponent.h b/Engine/Plugins/Experimental/VirtualHeightfieldMesh/Source/VirtualHeightfieldMesh/Public/VirtualHeightfieldMeshComponent.h
--- a/Engine/Plugins/Experimental/VirtualHeightfieldMesh/Source/VirtualHeightfieldMesh/Public/VirtualHeightfieldMeshComponent.h
+++ b/Engine/Plugins/Experimental/VirtualHeightfieldMesh/Source/VirtualHeightfieldMesh/Public/VirtualHeightfieldMeshComponent.h
@@ -78,6 +78,10 @@ protected:
 	UPROPERTY(EditAnywhere, Category = Rendering, meta = (DisplayName = "Actor Hidden In Editor", DisplayAfter = "bHiddenInGame"))
 	bool bHiddenInEditor = true;
 
+	/** Allows material World Position Offset to contribute to velocity rendering.*/
+	UPROPERTY(EditAnywhere, Category = Rendering, AdvancedDisplay)
+	bool bWorldPositionOffsetVelocity = false;
+
 public:
 	/** Get the HiddenInEditor flag  on this component. */
 	bool GetHiddenInEditor() const { return bHiddenInEditor; }
@@ -142,5 +146,6 @@ protected:
 	virtual void SetMaterial(int32 ElementIndex, class UMaterialInterface* Material) override;
 	virtual UMaterialInterface* GetMaterial(int32 Index) const override { return Material; }
 	virtual void GetUsedMaterials(TArray<UMaterialInterface*>& OutMaterials, bool bGetDebugMaterials = false) const override;
+	virtual bool SupportsWorldPositionOffsetVelocity() const override { return bWorldPositionOffsetVelocity; }
 	//~ End UPrimitiveComponent Interface
 };
