Fixed crash when using VHM with new init views task graph.

diff --git a/Engine/Plugins/Experimental/VirtualHeightfieldMesh/Source/VirtualHeightfieldMesh/Private/VirtualHeightfieldMeshSceneProxy.cpp b/Engine/Plugins/Experimental/VirtualHeightfieldMesh/Source/VirtualHeightfieldMesh/Private/VirtualHeightfieldMeshSceneProxy.cpp
--- a/Engine/Plugins/Experimental/VirtualHeightfieldMesh/Source/VirtualHeightfieldMesh/Private/VirtualHeightfieldMeshSceneProxy.cpp
+++ b/Engine/Plugins/Experimental/VirtualHeightfieldMesh/Source/VirtualHeightfieldMesh/Private/VirtualHeightfieldMeshSceneProxy.cpp
@@ -117,6 +117,7 @@ struct FOcclusionResults
 	FTexture2DRHIRef OcclusionTexture;
 	FIntPoint TextureSize;
 	int32 NumTextureMips;
+	TArray<bool> UploadData;
 };
 
 struct FOcclusionResultsKey
@@ -670,8 +671,6 @@ void FVirtualHeightfieldMeshSceneProxy::BuildOcclusionVolumes(TArrayView<FVector
 
 void FVirtualHeightfieldMeshSceneProxy::AcceptOcclusionResults(FSceneView const* View, TArray<bool>* Results, int32 ResultsStart, int32 NumResults)
 {
-	check(IsInRenderingThread());
-
 	if (GOcclusionResetRequired)
 	{
 		GOcclusionResults.Reset();
@@ -683,35 +682,8 @@ void FVirtualHeightfieldMeshSceneProxy::AcceptOcclusionResults(FSceneView const*
 		FOcclusionResults& OcclusionResults = GOcclusionResults.Emplace(FOcclusionResultsKey(this, View));
 		OcclusionResults.TextureSize = OcclusionGridSize;
 		OcclusionResults.NumTextureMips = NumOcclusionLods;
-
-		const FRHITextureCreateDesc Desc =
-			FRHITextureCreateDesc::Create2D(TEXT("VirtualHeightfieldMesh.OcclusionTexture"), OcclusionGridSize, PF_G8)
-			.SetNumMips(NumOcclusionLods)
-			.SetFlags(ETextureCreateFlags::ShaderResource);
-
-		OcclusionResults.OcclusionTexture = RHICreateTexture(Desc);
-		
-		bool const* Src = Results->GetData() + ResultsStart;
-		FIntPoint Size = OcclusionGridSize;
-		for (int32 MipIndex = 0; MipIndex < NumOcclusionLods; ++MipIndex)
-		{
-			uint32 Stride;
-			uint8* Dst = (uint8*)RHILockTexture2D(OcclusionResults.OcclusionTexture, MipIndex, RLM_WriteOnly, Stride, false);
-		
-			for (int Y = 0; Y < Size.Y; ++Y)
-			{
-				for (int X = 0; X < Size.X; ++X)
-				{
-					Dst[Y * Stride + X] = *(Src++) ? 255 : 0;
-				}
-			}
-			
-			RHIUnlockTexture2D(OcclusionResults.OcclusionTexture, MipIndex, false);
-
-			Size.X = FMath::Max(Size.X / 2, 1);
-			Size.Y = FMath::Max(Size.Y / 2, 1);
-		}
-	}		
+		OcclusionResults.UploadData.Append(Results->GetData() + ResultsStart, NumResults);
+	}
 }
 
 namespace VirtualHeightfieldMesh
@@ -1301,6 +1273,40 @@ void FVirtualHeightfieldMeshRendererExtension::SubmitWork(FRDGBuilder& GraphBuil
 			}
 			else
 			{
+				if (!OcclusionResults->UploadData.IsEmpty())
+				{
+					const FRHITextureCreateDesc Desc =
+						FRHITextureCreateDesc::Create2D(TEXT("VirtualHeightfieldMesh.OcclusionTexture"), OcclusionResults->TextureSize, PF_G8)
+						.SetNumMips(OcclusionResults->NumTextureMips)
+						.SetFlags(ETextureCreateFlags::ShaderResource);
+
+					OcclusionResults->OcclusionTexture = RHICreateTexture(Desc);
+
+					int32 SourceIndex = 0;
+
+					FIntPoint Size = OcclusionResults->TextureSize;
+					for (int32 MipIndex = 0; MipIndex < OcclusionResults->NumTextureMips; ++MipIndex)
+					{
+						uint32 Stride;
+						uint8* Dst = (uint8*)RHILockTexture2D(OcclusionResults->OcclusionTexture, MipIndex, RLM_WriteOnly, Stride, false);
+
+						for (int Y = 0; Y < Size.Y; ++Y)
+						{
+							for (int X = 0; X < Size.X; ++X)
+							{
+								Dst[Y * Stride + X] = OcclusionResults->UploadData[SourceIndex++] ? 255 : 0;
+							}
+						}
+
+						RHIUnlockTexture2D(OcclusionResults->OcclusionTexture, MipIndex, false);
+
+						Size.X = FMath::Max(Size.X / 2, 1);
+						Size.Y = FMath::Max(Size.Y / 2, 1);
+					}
+
+					OcclusionResults->UploadData.Empty();
+				}
+
 				MainViewDesc.OcclusionTexture = OcclusionResults->OcclusionTexture;
 				MainViewDesc.OcclusionLevelOffset = (int32)ProxyDesc.MaxLevel - OcclusionResults->NumTextureMips + 1;
 			}
