Modified scene proxy GetDynamicMeshElements calls to take a local command list instead of relying on the immediate command list and defer uniform expression cache invalidations.
 - Calls to FRHICommandListImmediate::Get() replaced with Collector.GetRHICommandList()

diff --git a/Engine/Plugins/2D/Paper2D/Source/Paper2D/Private/PaperRenderSceneProxy.cpp b/Engine/Plugins/2D/Paper2D/Source/Paper2D/Private/PaperRenderSceneProxy.cpp
--- a/Engine/Plugins/2D/Paper2D/Source/Paper2D/Private/PaperRenderSceneProxy.cpp
+++ b/Engine/Plugins/2D/Paper2D/Source/Paper2D/Private/PaperRenderSceneProxy.cpp
@@ -58,7 +58,7 @@ public:
 	{
 	}
 
-	void CheckValidity(FMaterialRenderProxy* InCurrentParent)
+	bool CheckValidity(FMaterialRenderProxy* InCurrentParent)
 	{
 		if (InCurrentParent != Parent)
 		{
@@ -66,12 +66,15 @@ public:
 			ParentMaterialSerialNumber = INDEX_NONE;
 		}
 
+		bool bValid = true;
+
 		if (ParentMaterialSerialNumber != Parent->GetExpressionCacheSerialNumber())
 		{
-			// Not valid, need to rebuild
-			CacheUniformExpressions(FRHICommandListImmediate::Get(), /*bRecreateUniformBuffer=*/ true);
+			bValid = false;
 			ParentMaterialSerialNumber = Parent->GetExpressionCacheSerialNumber();
 		}
+
+		return bValid;
 	}
 
 	void Reinitialize(const FMaterialRenderProxy* InParent, const UTexture* InBaseTexture, FAdditionalSpriteTextureArray InAdditionalTextures)
@@ -288,7 +291,6 @@ void FPaperRenderSceneProxy::DebugDrawCollision(const FSceneView* View, int32 Vi
 void FPaperRenderSceneProxy::GetDynamicMeshElements(const TArray<const FSceneView*>& Views, const FSceneViewFamily& ViewFamily, uint32 VisibilityMap, FMeshElementCollector& Collector) const
 {
 	SCOPE_CYCLE_COUNTER(STAT_PaperRenderSceneProxy_GetDynamicMeshElements);
-	checkSlow(IsInRenderingThread());
 
 	const FEngineShowFlags& EngineShowFlags = ViewFamily.EngineShowFlags;
 
@@ -404,7 +406,7 @@ void FPaperRenderSceneProxy::GetNewBatchMeshes(const FSceneView* View, int32 Vie
 				}
 
 				// Create a texture override material proxy or make sure our existing one is valid
-				FSpriteTextureOverrideRenderProxy* SectionMaterialProxy = GetCachedMaterialProxyForSection(SectionIndex, ParentMaterialProxy);
+				FSpriteTextureOverrideRenderProxy* SectionMaterialProxy = GetCachedMaterialProxyForSection(Collector, SectionIndex, ParentMaterialProxy);
 
 				Settings.CastShadow = bCastShadow;
 				Settings.bDisableBackfaceCulling = bDrawTwoSided;
@@ -430,7 +432,7 @@ void FPaperRenderSceneProxy::GetNewBatchMeshesPrebuilt(const FSceneView* View, i
 	{
 		FMeshBatch& Batch = Collector.AllocateMesh();
 
-		if (GetMeshElement(SectionIndex, DPG, IsSelected(), Batch))
+		if (GetMeshElement(Collector, SectionIndex, DPG, IsSelected(), Batch))
 		{
 			// Implementing our own wireframe coloring as the automatic one (controlled by Mesh.bCanApplyViewModeOverrides) only supports per-FPrimitiveSceneProxy WireframeColor
 			if (bIsWireframeView)
@@ -458,7 +460,7 @@ void FPaperRenderSceneProxy::GetNewBatchMeshesPrebuilt(const FSceneView* View, i
 	}
 }
 
-bool FPaperRenderSceneProxy::GetMeshElement(int32 SectionIndex, uint8 DepthPriorityGroup, bool bIsSelected, FMeshBatch& OutMeshBatch) const
+bool FPaperRenderSceneProxy::GetMeshElement(FMeshElementCollector& Collector, int32 SectionIndex, uint8 DepthPriorityGroup, bool bIsSelected, FMeshBatch& OutMeshBatch) const
 {
 	check(bSpritesUseVertexBufferPath);
 	check(SectionIndex < BatchedSections.Num());
@@ -479,7 +481,7 @@ bool FPaperRenderSceneProxy::GetMeshElement(int32 SectionIndex, uint8 DepthPrior
 		OutMeshBatch.DepthPriorityGroup = DepthPriorityGroup;
 		OutMeshBatch.Type = PT_TriangleList;
 		OutMeshBatch.bDisableBackfaceCulling = bDrawTwoSided;
-		OutMeshBatch.MaterialRenderProxy = GetCachedMaterialProxyForSection(SectionIndex, Section.Material->GetRenderProxy());
+		OutMeshBatch.MaterialRenderProxy = GetCachedMaterialProxyForSection(Collector, SectionIndex, Section.Material->GetRenderProxy());
 
 		// Set up the FMeshBatchElement.
 		FMeshBatchElement& BatchElement = OutMeshBatch.Elements[0];
@@ -598,7 +600,7 @@ void FPaperRenderSceneProxy::SetTransientTextureOverride_RenderThread(const UTex
 }
 #endif
 
-FSpriteTextureOverrideRenderProxy* FPaperRenderSceneProxy::GetCachedMaterialProxyForSection(int32 SectionIndex, FMaterialRenderProxy* ParentMaterialProxy) const
+FSpriteTextureOverrideRenderProxy* FPaperRenderSceneProxy::GetCachedMaterialProxyForSection(FMeshElementCollector& Collector, int32 SectionIndex, FMaterialRenderProxy* ParentMaterialProxy) const
 {
 	if (MaterialTextureOverrideProxies.Num() < BatchedSections.Num())
 	{
@@ -612,7 +614,10 @@ FSpriteTextureOverrideRenderProxy* FPaperRenderSceneProxy::GetCachedMaterialProx
 		Result = new FSpriteTextureOverrideRenderProxy(ParentMaterialProxy, Section.BaseTexture, Section.AdditionalTextures UE_EXPAND_IF_WITH_EDITOR(, TextureOverrideList));
 	}
 
-	Result->CheckValidity(ParentMaterialProxy);
+	if (!Result->CheckValidity(ParentMaterialProxy))
+	{
+		Collector.CacheUniformExpressions(Result, /** bRecreateUniformBuffer */ true);
+	}
 
 	return Result;
 }
diff --git a/Engine/Plugins/2D/Paper2D/Source/Paper2D/Private/PaperRenderSceneProxy.h b/Engine/Plugins/2D/Paper2D/Source/Paper2D/Private/PaperRenderSceneProxy.h
--- a/Engine/Plugins/2D/Paper2D/Source/Paper2D/Private/PaperRenderSceneProxy.h
+++ b/Engine/Plugins/2D/Paper2D/Source/Paper2D/Private/PaperRenderSceneProxy.h
@@ -145,7 +145,7 @@ public:
 protected:
 	virtual void GetDynamicMeshElementsForView(const FSceneView* View, int32 ViewIndex, FMeshElementCollector& Collector) const;
 
-	bool GetMeshElement(int32 SectionIndex, uint8 DepthPriorityGroup, bool bIsSelected, FMeshBatch& OutMeshBatch) const;
+	bool GetMeshElement(FMeshElementCollector& Collector, int32 SectionIndex, uint8 DepthPriorityGroup, bool bIsSelected, FMeshBatch& OutMeshBatch) const;
 
 	void GetNewBatchMeshes(const FSceneView* View, int32 ViewIndex, FMeshElementCollector& Collector) const;
 	void GetNewBatchMeshesPrebuilt(const FSceneView* View, int32 ViewIndex, FMeshElementCollector& Collector) const;
@@ -158,7 +158,7 @@ protected:
 	// Call this if you modify BatchedSections or Vertices after the proxy has already been created
 	void RecreateCachedRenderData(FRHICommandListBase& RHICmdList);
 
-	FSpriteTextureOverrideRenderProxy* GetCachedMaterialProxyForSection(int32 SectionIndex, FMaterialRenderProxy* ParentMaterialProxy) const;
+	FSpriteTextureOverrideRenderProxy* GetCachedMaterialProxyForSection(FMeshElementCollector& Collector, int32 SectionIndex, FMaterialRenderProxy* ParentMaterialProxy) const;
 
 protected:
 	TArray<FSpriteRenderSection> BatchedSections;
diff --git a/Engine/Plugins/2D/Paper2D/Source/Paper2D/Private/PaperTileMapRenderSceneProxy.cpp b/Engine/Plugins/2D/Paper2D/Source/Paper2D/Private/PaperTileMapRenderSceneProxy.cpp
--- a/Engine/Plugins/2D/Paper2D/Source/Paper2D/Private/PaperTileMapRenderSceneProxy.cpp
+++ b/Engine/Plugins/2D/Paper2D/Source/Paper2D/Private/PaperTileMapRenderSceneProxy.cpp
@@ -264,7 +264,6 @@ void FPaperTileMapRenderSceneProxy::DrawHexagonalGridLines(FPrimitiveDrawInterfa
 void FPaperTileMapRenderSceneProxy::GetDynamicMeshElements(const TArray<const FSceneView*>& Views, const FSceneViewFamily& ViewFamily, uint32 VisibilityMap, FMeshElementCollector& Collector) const
 {
 	SCOPE_CYCLE_COUNTER(STAT_TileMap_GetDynamicMeshElements);
-	checkSlow(IsInRenderingThread());
 
 	for (int32 ViewIndex = 0; ViewIndex < Views.Num(); ViewIndex++)
 	{
diff --git a/Engine/Plugins/Editor/UVEditor/Source/UVEditorTools/Private/Drawing/BasicLineSetComponent.cpp b/Engine/Plugins/Editor/UVEditor/Source/UVEditorTools/Private/Drawing/BasicLineSetComponent.cpp
--- a/Engine/Plugins/Editor/UVEditor/Source/UVEditorTools/Private/Drawing/BasicLineSetComponent.cpp
+++ b/Engine/Plugins/Editor/UVEditor/Source/UVEditorTools/Private/Drawing/BasicLineSetComponent.cpp
@@ -154,7 +154,7 @@ public:
 					Mesh.MaterialRenderProxy = MeshBatchData.MaterialProxy;
 
 					FDynamicPrimitiveUniformBuffer& DynamicPrimitiveUniformBuffer = Collector.AllocateOneFrameResource<FDynamicPrimitiveUniformBuffer>();
-					DynamicPrimitiveUniformBuffer.Set(GetLocalToWorld(), GetLocalToWorld(), GetBounds(), GetLocalBounds(), true, false, AlwaysHasVelocity());
+					DynamicPrimitiveUniformBuffer.Set(Collector.GetRHICommandList(), GetLocalToWorld(), GetLocalToWorld(), GetBounds(), GetLocalBounds(), true, false, AlwaysHasVelocity());
 					BatchElement.PrimitiveUniformBufferResource = &DynamicPrimitiveUniformBuffer.UniformBuffer;
 
 					BatchElement.FirstIndex = MeshBatchData.StartIndex;
diff --git a/Engine/Plugins/Editor/UVEditor/Source/UVEditorTools/Private/Drawing/BasicPointSetComponent.cpp b/Engine/Plugins/Editor/UVEditor/Source/UVEditorTools/Private/Drawing/BasicPointSetComponent.cpp
--- a/Engine/Plugins/Editor/UVEditor/Source/UVEditorTools/Private/Drawing/BasicPointSetComponent.cpp
+++ b/Engine/Plugins/Editor/UVEditor/Source/UVEditorTools/Private/Drawing/BasicPointSetComponent.cpp
@@ -146,7 +146,7 @@ public:
 					Mesh.MaterialRenderProxy = MeshBatchData.MaterialProxy;
 
 					FDynamicPrimitiveUniformBuffer& DynamicPrimitiveUniformBuffer = Collector.AllocateOneFrameResource<FDynamicPrimitiveUniformBuffer>();
-					DynamicPrimitiveUniformBuffer.Set(GetLocalToWorld(), GetLocalToWorld(), GetBounds(), GetLocalBounds(), true, false, AlwaysHasVelocity());
+					DynamicPrimitiveUniformBuffer.Set(Collector.GetRHICommandList(), GetLocalToWorld(), GetLocalToWorld(), GetBounds(), GetLocalBounds(), true, false, AlwaysHasVelocity());
 					BatchElement.PrimitiveUniformBufferResource = &DynamicPrimitiveUniformBuffer.UniformBuffer;
 
 					BatchElement.FirstIndex = MeshBatchData.StartIndex;
diff --git a/Engine/Plugins/Editor/UVEditor/Source/UVEditorTools/Private/Drawing/BasicTriangleSetComponent.cpp b/Engine/Plugins/Editor/UVEditor/Source/UVEditorTools/Private/Drawing/BasicTriangleSetComponent.cpp
--- a/Engine/Plugins/Editor/UVEditor/Source/UVEditorTools/Private/Drawing/BasicTriangleSetComponent.cpp
+++ b/Engine/Plugins/Editor/UVEditor/Source/UVEditorTools/Private/Drawing/BasicTriangleSetComponent.cpp
@@ -143,7 +143,7 @@ public:
 					Mesh.MaterialRenderProxy = MeshBatchData.MaterialProxy;
 
 					FDynamicPrimitiveUniformBuffer& DynamicPrimitiveUniformBuffer = Collector.AllocateOneFrameResource<FDynamicPrimitiveUniformBuffer>();
-					DynamicPrimitiveUniformBuffer.Set(GetLocalToWorld(), GetLocalToWorld(), GetBounds(), GetLocalBounds(), true, false, AlwaysHasVelocity());
+					DynamicPrimitiveUniformBuffer.Set(Collector.GetRHICommandList(), GetLocalToWorld(), GetLocalToWorld(), GetBounds(), GetLocalBounds(), true, false, AlwaysHasVelocity());
 					BatchElement.PrimitiveUniformBufferResource = &DynamicPrimitiveUniformBuffer.UniformBuffer;
 
 					BatchElement.FirstIndex = MeshBatchData.StartIndex;
diff --git a/Engine/Plugins/Experimental/Dataflow/Source/DataflowEnginePlugin/Private/Dataflow/DataflowEngineSceneProxy.cpp b/Engine/Plugins/Experimental/Dataflow/Source/DataflowEnginePlugin/Private/Dataflow/DataflowEngineSceneProxy.cpp
--- a/Engine/Plugins/Experimental/Dataflow/Source/DataflowEnginePlugin/Private/Dataflow/DataflowEngineSceneProxy.cpp
+++ b/Engine/Plugins/Experimental/Dataflow/Source/DataflowEnginePlugin/Private/Dataflow/DataflowEngineSceneProxy.cpp
@@ -124,7 +124,6 @@ void FDataflowEngineSceneProxy::GetDynamicMeshElements(const TArray<const FScene
 	QUICK_SCOPE_CYCLE_COUNTER(STAT_OverlaySceneProxy_GetDynamicMeshElements);
 
 	check(RenderMaterial);
-	check(IsInRenderingThread());
 
 	for (int32 ViewIndex = 0; ViewIndex < Views.Num(); ViewIndex++)
 	{
@@ -268,7 +267,7 @@ void FDataflowEngineSceneProxy::GetMeshDynamicMeshElements(int32 ViewIndex, FMes
 	for (const FDataflowTriangleSetMeshBatchData& MeshBatchData : MeshBatchDatas)
 	{
 		FDynamicPrimitiveUniformBuffer& DynamicPrimitiveUniformBuffer = Collector.AllocateOneFrameResource<FDynamicPrimitiveUniformBuffer>();
-		DynamicPrimitiveUniformBuffer.Set(GetLocalToWorld(), GetLocalToWorld(), GetBounds(), GetLocalBounds(), true, false, AlwaysHasVelocity());
+		DynamicPrimitiveUniformBuffer.Set(Collector.GetRHICommandList(), GetLocalToWorld(), GetLocalToWorld(), GetBounds(), GetLocalBounds(), true, false, AlwaysHasVelocity());
 
 		FMeshBatch& Mesh = Collector.AllocateMesh();
 		Mesh.bWireframe = false;
@@ -401,7 +400,6 @@ void FDataflowEngineSceneProxy::CreateInstancedVertexRenderThreadResources(FRHIC
 
 void FDataflowEngineSceneProxy::DestroyInstancedVertexRenderThreadResources()
 {
-	check(IsInRenderingThread());
 	if (NumRenderedVerts)
 	{
 		BoxVertexBuffers.PositionVertexBuffer.ReleaseResource();
@@ -427,7 +425,7 @@ void FDataflowEngineSceneProxy::GetDynamicInstancedVertexMeshElements(int32 View
 			for (const FDataflowVertexBatchData& VertexBatchData : VertexBatchDatas)
 			{
 				FDynamicPrimitiveUniformBuffer& DynamicPrimitiveUniformBuffer = Collector.AllocateOneFrameResource<FDynamicPrimitiveUniformBuffer>();
-				DynamicPrimitiveUniformBuffer.Set(GetLocalToWorld(), GetLocalToWorld(), GetBounds(), GetLocalBounds(), true, false, AlwaysHasVelocity());
+				DynamicPrimitiveUniformBuffer.Set(Collector.GetRHICommandList(), GetLocalToWorld(), GetLocalToWorld(), GetBounds(), GetLocalBounds(), true, false, AlwaysHasVelocity());
 
 				FMeshBatch& Mesh = Collector.AllocateMesh();
 				Mesh.CastShadow = false;
diff --git a/Engine/Plugins/Experimental/ImagePlate/Source/ImagePlate/Private/ImagePlateComponent.cpp b/Engine/Plugins/Experimental/ImagePlate/Source/ImagePlate/Private/ImagePlateComponent.cpp
--- a/Engine/Plugins/Experimental/ImagePlate/Source/ImagePlate/Private/ImagePlateComponent.cpp
+++ b/Engine/Plugins/Experimental/ImagePlate/Source/ImagePlate/Private/ImagePlateComponent.cpp
@@ -108,7 +108,7 @@ namespace
 
 		virtual void CreateRenderThreadResources(FRHICommandListBase& RHICmdList) override
 		{
-			BuildMesh();
+			BuildMesh(RHICmdList);
 			IndexBuffer.InitResource(RHICmdList);
 #if RHI_RAYTRACING
 			if (IsRayTracingEnabled())
@@ -134,7 +134,7 @@ namespace
 #endif
 		}
 
-		void BuildMesh()
+		void BuildMesh(FRHICommandListBase& RHICmdList)
 		{
 			TArray<FDynamicMeshVertex> Vertices;
 			Vertices.Empty(4);
@@ -151,7 +151,7 @@ namespace
 			Vertices[2].TextureCoordinate[0] = FVector2f(1,0);
 			Vertices[3].TextureCoordinate[0] = FVector2f(1,1);
 
-			VertexBuffers.InitFromDynamicVertex(&VertexFactory, Vertices);
+			VertexBuffers.InitFromDynamicVertex(RHICmdList, &VertexFactory, Vertices);
 
 			IndexBuffer.Indices.Empty(6);
 			IndexBuffer.Indices.AddUninitialized(6);
@@ -314,7 +314,7 @@ namespace
 				bOutputVelocity |= AlwaysHasVelocity();
 
 				FDynamicPrimitiveUniformBuffer& DynamicPrimitiveUniformBuffer = Context.RayTracingMeshResourceCollector.AllocateOneFrameResource<FDynamicPrimitiveUniformBuffer>();
-				DynamicPrimitiveUniformBuffer.Set(GetLocalToWorld(), PreviousLocalToWorld, GetBounds(), GetLocalBounds(), GetLocalBounds(), true, bHasPrecomputedVolumetricLightmap, bOutputVelocity, GetCustomPrimitiveData());
+				DynamicPrimitiveUniformBuffer.Set(Context.RHICmdList, GetLocalToWorld(), PreviousLocalToWorld, GetBounds(), GetLocalBounds(), GetLocalBounds(), true, bHasPrecomputedVolumetricLightmap, bOutputVelocity, GetCustomPrimitiveData());
 				BatchElement.PrimitiveUniformBufferResource = &DynamicPrimitiveUniformBuffer.UniformBuffer;
 
 				BatchElement.FirstIndex = 0;
diff --git a/Engine/Plugins/Experimental/VirtualHeightfieldMesh/Source/VirtualHeightfieldMesh/Private/VirtualHeightfieldMeshSceneProxy.cpp b/Engine/Plugins/Experimental/VirtualHeightfieldMesh/Source/VirtualHeightfieldMesh/Private/VirtualHeightfieldMeshSceneProxy.cpp
--- a/Engine/Plugins/Experimental/VirtualHeightfieldMesh/Source/VirtualHeightfieldMesh/Private/VirtualHeightfieldMeshSceneProxy.cpp
+++ b/Engine/Plugins/Experimental/VirtualHeightfieldMesh/Source/VirtualHeightfieldMesh/Private/VirtualHeightfieldMeshSceneProxy.cpp
@@ -21,6 +21,7 @@
 #include "VirtualHeightfieldMeshVertexFactory.h"
 #include "VT/RuntimeVirtualTexture.h"
 #include "VT/VirtualTextureFeedbackBuffer.h"
+#include "Async/Mutex.h"
 
 DECLARE_STATS_GROUP(TEXT("VirtualHeightfieldMesh"), STATGROUP_VirtualHeightfieldMesh, STATCAT_Advanced);
 DECLARE_CYCLE_STAT(TEXT("VirtualHeightfieldMesh SubmitWork"), STAT_VirtualHeightfieldMesh_SubmitWork, STATGROUP_VirtualHeightfieldMesh);
@@ -99,7 +100,7 @@ namespace VirtualHeightfieldMesh
 	};
 
 	/** Initialize the FDrawInstanceBuffers objects. */
-	void InitializeInstanceBuffers(FRHICommandListImmediate& InRHICmdList, FDrawInstanceBuffers& InBuffers);
+	void InitializeInstanceBuffers(FRHICommandListBase& InRHICmdList, FDrawInstanceBuffers& InBuffers);
 
 	/** Release the FDrawInstanceBuffers objects. */
 	void ReleaseInstanceBuffers(FDrawInstanceBuffers& InBuffers)
@@ -183,7 +184,7 @@ public:
 	bool IsInFrame() { return bInFrame; }
 
 	/** Call once per frame for each mesh/view that has relevance. This allocates the buffers to use for the frame and adds the work to fill the buffers to the queue. */
-	VirtualHeightfieldMesh::FDrawInstanceBuffers& AddWork(FVirtualHeightfieldMeshSceneProxy const* InProxy, FSceneView const* InMainView, FSceneView const* InCullView);
+	VirtualHeightfieldMesh::FDrawInstanceBuffers& AddWork(FRHICommandListBase& RHICmdList, FVirtualHeightfieldMeshSceneProxy const* InProxy, FSceneView const* InMainView, FSceneView const* InCullView);
 	/** Submit all the work added by AddWork(). The work fills all of the buffers ready for use by the referencing mesh batches. */
 	void SubmitWork(FRDGBuilder& GraphBuilder);
 
@@ -202,6 +203,8 @@ private:
 	/** Flag for frame validation. */
 	bool bInFrame;
 
+	UE::FMutex Mutex;
+
 	/** Buffers to fill. Resources can persist between frames to reduce allocation cost, but contents don't persist. */
 	TArray<VirtualHeightfieldMesh::FDrawInstanceBuffers> Buffers;
 	/** Per buffer frame time stamp of last usage. */
@@ -262,8 +265,10 @@ void FVirtualHeightfieldMeshRendererExtension::ReleaseRHI()
 	Buffers.Empty();
 }
 
-VirtualHeightfieldMesh::FDrawInstanceBuffers& FVirtualHeightfieldMeshRendererExtension::AddWork(FVirtualHeightfieldMeshSceneProxy const* InProxy, FSceneView const* InMainView, FSceneView const* InCullView)
+VirtualHeightfieldMesh::FDrawInstanceBuffers& FVirtualHeightfieldMeshRendererExtension::AddWork(FRHICommandListBase& RHICmdList, FVirtualHeightfieldMeshSceneProxy const* InProxy, FSceneView const* InMainView, FSceneView const* InCullView)
 {
+	UE::TScopeLock Lock(Mutex);
+
 	// If we hit this then BegineFrame()/EndFrame() logic needs fixing in the Scene Renderer.
 	if (!ensure(!bInFrame))
 	{
@@ -308,7 +313,7 @@ VirtualHeightfieldMesh::FDrawInstanceBuffers& FVirtualHeightfieldMeshRendererExt
 		DiscardIds.Add(DiscardId);
 		WorkDesc.BufferIndex = Buffers.AddDefaulted();
 		WorkDescs.Add(WorkDesc);
-		VirtualHeightfieldMesh::InitializeInstanceBuffers(GetImmediateCommandList_ForRenderCommand(), Buffers[WorkDesc.BufferIndex]);
+		VirtualHeightfieldMesh::InitializeInstanceBuffers(RHICmdList, Buffers[WorkDesc.BufferIndex]);
 	}
 
 	return Buffers[WorkDesc.BufferIndex];
@@ -316,6 +321,7 @@ VirtualHeightfieldMesh::FDrawInstanceBuffers& FVirtualHeightfieldMeshRendererExt
 
 void FVirtualHeightfieldMeshRendererExtension::BeginFrame(FRDGBuilder& GraphBuilder)
 {
+	UE::TScopeLock Lock(Mutex);
 	// If we hit this then BegineFrame()/EndFrame() logic needs fixing in the Scene Renderer.
 	if (!ensure(!bInFrame))
 	{
@@ -331,6 +337,7 @@ void FVirtualHeightfieldMeshRendererExtension::BeginFrame(FRDGBuilder& GraphBuil
 
 void FVirtualHeightfieldMeshRendererExtension::EndFrame()
 {
+	UE::TScopeLock Lock(Mutex);
 	ensure(bInFrame);
 	bInFrame = false;
 
@@ -537,7 +544,6 @@ FPrimitiveViewRelevance FVirtualHeightfieldMeshSceneProxy::GetViewRelevance(cons
 
 void FVirtualHeightfieldMeshSceneProxy::GetDynamicMeshElements(const TArray<const FSceneView*>& Views, const FSceneViewFamily& ViewFamily, uint32 VisibilityMap, FMeshElementCollector& Collector) const
 {
-	check(IsInRenderingThread());
 	check(AllocatedVirtualTexture != nullptr);
 
 	if (GVirtualHeightfieldMeshViewRendererExtension.IsInFrame())
@@ -554,7 +560,7 @@ void FVirtualHeightfieldMeshSceneProxy::GetDynamicMeshElements(const TArray<cons
 	{
 		if (VisibilityMap & (1 << ViewIndex))
 		{
-			VirtualHeightfieldMesh::FDrawInstanceBuffers& Buffers = GVirtualHeightfieldMeshViewRendererExtension.AddWork(this, ViewFamily.Views[0], Views[ViewIndex]);
+			VirtualHeightfieldMesh::FDrawInstanceBuffers& Buffers = GVirtualHeightfieldMeshViewRendererExtension.AddWork(Collector.GetRHICommandList(), this, ViewFamily.Views[0], Views[ViewIndex]);
 
 			FMeshBatch& Mesh = Collector.AllocateMesh();
 			Mesh.bWireframe = AllowDebugViewmodes() && ViewFamily.EngineShowFlags.Wireframe;
@@ -981,7 +987,7 @@ namespace VirtualHeightfieldMesh
 	};
 
 	/** Initialize the FDrawInstanceBuffers objects. */
-	void InitializeInstanceBuffers(FRHICommandListImmediate& RHICmdList, FDrawInstanceBuffers& InBuffers)
+	void InitializeInstanceBuffers(FRHICommandListBase& RHICmdList, FDrawInstanceBuffers& InBuffers)
 	{
 		{
 			FRHIResourceCreateInfo CreateInfo(TEXT("VirtualHeightfieldMesh.InstanceBuffer"));
diff --git a/Engine/Plugins/Experimental/Water/Source/Runtime/Private/WaterMeshSceneProxy.cpp b/Engine/Plugins/Experimental/Water/Source/Runtime/Private/WaterMeshSceneProxy.cpp
--- a/Engine/Plugins/Experimental/Water/Source/Runtime/Private/WaterMeshSceneProxy.cpp
+++ b/Engine/Plugins/Experimental/Water/Source/Runtime/Private/WaterMeshSceneProxy.cpp
@@ -174,7 +174,7 @@ void FWaterMeshSceneProxy::GetDynamicMeshElements(const TArray<const FSceneView*
 {
 	CSV_SCOPED_TIMING_STAT_EXCLUSIVE(Water);
 	TRACE_CPUPROFILER_EVENT_SCOPE(FWaterMeshSceneProxy::GetDynamicMeshElements);
-	FRHICommandListBase& RHICmdList = FRHICommandListImmediate::Get();
+	FRHICommandListBase& RHICmdList = Collector.GetRHICommandList();
 
 	// The water render groups we have to render for this batch : 
 	TArray<EWaterMeshRenderGroupType, TInlineAllocator<WaterVertexFactoryType::NumRenderGroups>> BatchRenderGroups;
diff --git a/Engine/Plugins/FX/Niagara/Source/Niagara/Private/DataInterface/NiagaraDataInterfaceDynamicMesh.cpp b/Engine/Plugins/FX/Niagara/Source/Niagara/Private/DataInterface/NiagaraDataInterfaceDynamicMesh.cpp
--- a/Engine/Plugins/FX/Niagara/Source/Niagara/Private/DataInterface/NiagaraDataInterfaceDynamicMesh.cpp
+++ b/Engine/Plugins/FX/Niagara/Source/Niagara/Private/DataInterface/NiagaraDataInterfaceDynamicMesh.cpp
@@ -623,7 +623,7 @@ namespace NDIDynamicMeshLocal
 			}
 		}
 
-		virtual void SetupVertexFactory(class FNiagaraMeshVertexFactory& VertexFactory, const FLODModelData& LODModelData) const override
+		virtual void SetupVertexFactory(FRHICommandListBase& RHICmdList, class FNiagaraMeshVertexFactory& VertexFactory, const FLODModelData& LODModelData) const override
 		{
 			const FNDIInstanceData_RenderThread* InstanceData = &OwnerProxy->InstanceData_RT.FindChecked(SystemInstanceID);
 
@@ -686,7 +686,7 @@ namespace NDIDynamicMeshLocal
 				StaticMeshData.ColorComponentsSRV = GNullColorVertexBuffer.VertexBufferSRV;
 			}
 
-			VertexFactory.SetData(StaticMeshData);
+			VertexFactory.SetData(RHICmdList, StaticMeshData);
 		}
 
 		virtual void GetUsedMaterials(TArray<UMaterialInterface*>& OutMaterials) const override
diff --git a/Engine/Plugins/FX/Niagara/Source/Niagara/Private/NiagaraComponent.cpp b/Engine/Plugins/FX/Niagara/Source/Niagara/Private/NiagaraComponent.cpp
--- a/Engine/Plugins/FX/Niagara/Source/Niagara/Private/NiagaraComponent.cpp
+++ b/Engine/Plugins/FX/Niagara/Source/Niagara/Private/NiagaraComponent.cpp
@@ -309,10 +309,8 @@ const FVector3f& FNiagaraSceneProxy::GetLWCRenderTile() const
 	return RenderData ? RenderData->LWCRenderTile : FVector3f::ZeroVector;
 }
 
-TUniformBuffer<FPrimitiveUniformShaderParameters>* FNiagaraSceneProxy::GetCustomUniformBufferResource(bool bHasVelocity, const FBox& InstanceBounds) const
+TUniformBuffer<FPrimitiveUniformShaderParameters>* FNiagaraSceneProxy::GetCustomUniformBufferResource(FRHICommandListBase& RHICmdList, bool bHasVelocity, const FBox& InstanceBounds) const
 {
-	FRHICommandListBase& RHICmdList = FRHICommandListImmediate::Get();
-
 	// Use a hash to determine if we can re-use any uniform buffer
 	uint32 KeyHash = HashCombine(bHasVelocity, InstanceBounds.IsValid);
 
@@ -389,7 +387,12 @@ TUniformBuffer<FPrimitiveUniformShaderParameters>* FNiagaraSceneProxy::GetCustom
 	return CustomUBRef;
 }
 
-FRHIUniformBuffer* FNiagaraSceneProxy::GetCustomUniformBuffer(bool bHasVelocity, const FBox& InstanceBounds) const
+TUniformBuffer<FPrimitiveUniformShaderParameters>* FNiagaraSceneProxy::GetCustomUniformBufferResource(bool bHasVelocity, const FBox& InstanceBounds) const
+{
+	return GetCustomUniformBufferResource(FRHICommandListImmediate::Get(), bHasVelocity, InstanceBounds);
+}
+
+FRHIUniformBuffer* FNiagaraSceneProxy::GetCustomUniformBuffer(FRHICommandListBase& RHICmdList, bool bHasVelocity, const FBox& InstanceBounds) const
 {
 	// Default UB we create for the primitive
 	if (bHasVelocity && !InstanceBounds.IsValid)
@@ -397,7 +400,12 @@ FRHIUniformBuffer* FNiagaraSceneProxy::GetCustomUniformBuffer(bool bHasVelocity,
 		return GetUniformBuffer();
 	}
 
-	return GetCustomUniformBufferResource(bHasVelocity, InstanceBounds)->GetUniformBufferRHI();
+	return GetCustomUniformBufferResource(RHICmdList, bHasVelocity, InstanceBounds)->GetUniformBufferRHI();
+}
+
+FRHIUniformBuffer* FNiagaraSceneProxy::GetCustomUniformBuffer(bool bHasVelocity, const FBox& InstanceBounds) const
+{
+	return GetCustomUniformBuffer(FRHICommandListImmediate::Get(), bHasVelocity, InstanceBounds);
 }
 
 uint32 FNiagaraSceneProxy::GetMemoryFootprint() const
@@ -457,9 +465,11 @@ void FNiagaraSceneProxy::GetDynamicMeshElements(const TArray<const FSceneView*>&
 		bEndTraceEvent = true;
 		FCpuProfilerTrace::OutputBeginDynamicEvent(GetResourceName());
 	}
-	else
 #endif //CPUPROFILERTRACE_ENABLED
 #if ENABLE_STATNAMEDEVENTS
+#if CPUPROFILERTRACE_ENABLED
+	else
+#endif
 	{
 		bEndPlatformEvent = true;
 		FPlatformMisc::BeginNamedEvent(FColor(0), *SystemStatString);
diff --git a/Engine/Plugins/FX/Niagara/Source/Niagara/Private/NiagaraGpuComputeDispatch.cpp b/Engine/Plugins/FX/Niagara/Source/Niagara/Private/NiagaraGpuComputeDispatch.cpp
--- a/Engine/Plugins/FX/Niagara/Source/Niagara/Private/NiagaraGpuComputeDispatch.cpp
+++ b/Engine/Plugins/FX/Niagara/Source/Niagara/Private/NiagaraGpuComputeDispatch.cpp
@@ -2187,9 +2187,9 @@ void FNiagaraGpuComputeDispatch::OnDestroy()
 	FFXSystemInterface::OnDestroy();
 }
 
-bool FNiagaraGpuComputeDispatch::AddSortedGPUSimulation(FNiagaraGPUSortInfo& SortInfo)
+bool FNiagaraGpuComputeDispatch::AddSortedGPUSimulation(FRHICommandListBase& RHICmdList, FNiagaraGPUSortInfo& SortInfo)
 {
-	if (GPUSortManager && GPUSortManager->AddTask(SortInfo.AllocationInfo, SortInfo.ParticleCount, SortInfo.SortFlags))
+	if (GPUSortManager && GPUSortManager->AddTask(RHICmdList, SortInfo.AllocationInfo, SortInfo.ParticleCount, SortInfo.SortFlags))
 	{
 		// It's not worth currently to have a map between SortInfo.AllocationInfo.SortBatchId and the relevant indices in SimulationsToSort
 		// because the number of batches is expect to be very small (1 or 2). If this change, it might be worth reconsidering.
diff --git a/Engine/Plugins/FX/Niagara/Source/Niagara/Private/NiagaraGpuComputeDispatch.h b/Engine/Plugins/FX/Niagara/Source/Niagara/Private/NiagaraGpuComputeDispatch.h
--- a/Engine/Plugins/FX/Niagara/Source/Niagara/Private/NiagaraGpuComputeDispatch.h
+++ b/Engine/Plugins/FX/Niagara/Source/Niagara/Private/NiagaraGpuComputeDispatch.h
@@ -141,7 +141,7 @@ public:
 	/** Processes all pending readbacks */
 	virtual void ProcessDebugReadbacks(FRHICommandListImmediate& RHICmdList, bool bWaitCompletion) override;
 
-	virtual bool AddSortedGPUSimulation(FNiagaraGPUSortInfo& SortInfo) override;
+	virtual bool AddSortedGPUSimulation(FRHICommandListBase& RHICmdList, FNiagaraGPUSortInfo& SortInfo) override;
 	virtual const FGlobalDistanceFieldParameterData* GetGlobalDistanceFieldData() const override;
 
 	void ResetDataInterfaces(FRDGBuilder& GraphBuilder, const FNiagaraGPUSystemTick& Tick, const FNiagaraComputeInstanceData& InstanceData) const;
diff --git a/Engine/Plugins/FX/Niagara/Source/Niagara/Private/NiagaraMeshRendererProperties.cpp b/Engine/Plugins/FX/Niagara/Source/Niagara/Private/NiagaraMeshRendererProperties.cpp
--- a/Engine/Plugins/FX/Niagara/Source/Niagara/Private/NiagaraMeshRendererProperties.cpp
+++ b/Engine/Plugins/FX/Niagara/Source/Niagara/Private/NiagaraMeshRendererProperties.cpp
@@ -120,12 +120,12 @@ public:
 		VertexBuffers.ColorVertexBuffer.BindColorVertexBuffer(VertexFactory, OutData);
 	}
 
-	void SetupVertexFactory(FNiagaraMeshVertexFactory& InVertexFactory, const FLODModelData& LODModelData) const override
+	void SetupVertexFactory(FRHICommandListBase& RHICmdList, FNiagaraMeshVertexFactory& InVertexFactory, const FLODModelData& LODModelData) const override
 	{
 		FStaticMeshDataType Data;
 		const FStaticMeshLODResources& LODResources = RenderData->LODResources[LODModelData.LODIndex];
 		InitVertexFactoryComponents(LODResources.VertexBuffers, &InVertexFactory, Data);
-		InVertexFactory.SetData(Data);
+		InVertexFactory.SetData(RHICmdList, Data);
 	}
 
 	void GetUsedMaterials(TArray<UMaterialInterface*>& OutMaterials) const override
diff --git a/Engine/Plugins/FX/Niagara/Source/Niagara/Private/NiagaraRendererMeshes.cpp b/Engine/Plugins/FX/Niagara/Source/Niagara/Private/NiagaraRendererMeshes.cpp
--- a/Engine/Plugins/FX/Niagara/Source/Niagara/Private/NiagaraRendererMeshes.cpp
+++ b/Engine/Plugins/FX/Niagara/Source/Niagara/Private/NiagaraRendererMeshes.cpp
@@ -712,7 +712,7 @@ uint32 FNiagaraRendererMeshes::PerformSortAndCull(FRHICommandListBase& RHICmdLis
 		if (ParticleMeshRenderData.bSortCullOnGpu)
 		{
 			SortInfo.CulledGPUParticleCountOffset = ParticleMeshRenderData.bNeedsCull ? ComputeDispatchInterface->GetGPUInstanceCounterManager().AcquireCulledEntry() : INDEX_NONE;
-			if (ComputeDispatchInterface->AddSortedGPUSimulation(SortInfo))
+			if (ComputeDispatchInterface->AddSortedGPUSimulation(RHICmdList, SortInfo))
 			{
 				ParticleMeshRenderData.ParticleSortedIndicesSRV = SortInfo.AllocationInfo.BufferSRV;
 				ParticleMeshRenderData.ParticleSortedIndicesOffset = SortInfo.AllocationInfo.BufferOffset;
@@ -1187,11 +1187,11 @@ void FNiagaraRendererMeshes::CreateMeshBatchForSection(
 	FMeshBatchElement& BatchElement = MeshBatch.Elements[0];
 	if (ParticleMeshRenderData.bUseGPUScene)
 	{
-		BatchElement.PrimitiveUniformBufferResource = SceneProxy.GetCustomUniformBufferResource(IsMotionBlurEnabled(), LocalBounds);
+		BatchElement.PrimitiveUniformBufferResource = SceneProxy.GetCustomUniformBufferResource(RHICmdList, IsMotionBlurEnabled(), LocalBounds);
 	}
 	else
 	{
-		BatchElement.PrimitiveUniformBuffer = SceneProxy.GetCustomUniformBuffer(IsMotionBlurEnabled(), LocalBounds);
+		BatchElement.PrimitiveUniformBuffer = SceneProxy.GetCustomUniformBuffer(RHICmdList, IsMotionBlurEnabled(), LocalBounds);
 	}
 	BatchElement.FirstIndex = 0;
 	BatchElement.MinVertexIndex = 0;
@@ -1275,7 +1275,7 @@ void FNiagaraRendererMeshes::GetDynamicMeshElements(const TArray<const FSceneVie
 {
 	check(SceneProxy);
 	PARTICLE_PERF_STAT_CYCLES_RT(SceneProxy->GetProxyDynamicData().PerfStatsContext, GetDynamicMeshElements);
-	FRHICommandListBase& RHICmdList = FRHICommandListImmediate::Get();
+	FRHICommandListBase& RHICmdList = Collector.GetRHICommandList();
 
 	// Prepare our particle render data
 	// This will also determine if we have anything to render
@@ -1390,8 +1390,8 @@ void FNiagaraRendererMeshes::GetDynamicMeshElements(const TArray<const FSceneVie
 					VertexFactory.SetNiagaraMeshFeatureLevel(ViewFamily.Scene->GetFeatureLevel());
 				}
 				VertexFactory.EnablePrimitiveIDElement(ParticleMeshRenderData.bUseGPUScene);
-				VertexFactory.InitResource(FRHICommandListImmediate::Get());
-				MeshData.RenderableMesh->SetupVertexFactory(VertexFactory, LODModel);
+				VertexFactory.InitResource(RHICmdList);
+				MeshData.RenderableMesh->SetupVertexFactory(RHICmdList, VertexFactory, LODModel);
 
 				PreparePerMeshData(ParticleMeshRenderData, VertexFactory, *SceneProxy, MeshData);
 
@@ -1562,8 +1562,8 @@ void FNiagaraRendererMeshes::GetDynamicRayTracingInstances(FRayTracingMaterialGa
 			VertexFactory.SetNiagaraMeshFeatureLevel(View->Family->Scene->GetFeatureLevel());
 		}
 		VertexFactory.EnablePrimitiveIDElement(ParticleMeshRenderData.bUseGPUScene);
-		VertexFactory.InitResource(Context.GraphBuilder.RHICmdList);
-		MeshData.RenderableMesh->SetupVertexFactory(VertexFactory, LODModel);
+		VertexFactory.InitResource(RHICmdList);
+		MeshData.RenderableMesh->SetupVertexFactory(RHICmdList, VertexFactory, LODModel);
 
 		PreparePerMeshData(ParticleMeshRenderData, VertexFactory, *SceneProxy, MeshData);
 
diff --git a/Engine/Plugins/FX/Niagara/Source/Niagara/Private/NiagaraRendererRibbons.cpp b/Engine/Plugins/FX/Niagara/Source/Niagara/Private/NiagaraRendererRibbons.cpp
--- a/Engine/Plugins/FX/Niagara/Source/Niagara/Private/NiagaraRendererRibbons.cpp
+++ b/Engine/Plugins/FX/Niagara/Source/Niagara/Private/NiagaraRendererRibbons.cpp
@@ -825,7 +825,7 @@ void FNiagaraRendererRibbons::GetDynamicMeshElements(const TArray<const FSceneVi
 		}
 	}
 
-	FRHICommandListBase& RHICmdList = FRHICommandListImmediate::Get();
+	FRHICommandListBase& RHICmdList = Collector.GetRHICommandList();
 
 #if STATS
 	FScopeCycleCounter EmitterStatsCounter(EmitterStatID);
@@ -982,7 +982,7 @@ void FNiagaraRendererRibbons::GetDynamicRayTracingInstances(FRayTracingMaterialG
 	SCOPE_CYCLE_COUNTER(STAT_NiagaraRenderRibbons);
 	check(SceneProxy);
 
-	FRHICommandListBase& RHICmdList = Context.GraphBuilder.RHICmdList;
+	FRHICommandListBase& RHICmdList = Context.RHICmdList;
 	FNiagaraDynamicDataRibbon *DynamicDataRibbon = static_cast<FNiagaraDynamicDataRibbon*>(DynamicDataRender);
 	FNiagaraGpuComputeDispatchInterface* ComputeDispatchInterface = SceneProxy->GetComputeDispatchInterface();
 	
@@ -2118,7 +2118,7 @@ inline void FNiagaraRendererRibbons::SetupMeshBatchAndCollectorResourceForView(F
 	}	
 	
 	// TODO: MotionVector/Velocity? Probably need to look into this?
-	MeshElement.PrimitiveUniformBuffer = SceneProxy->GetCustomUniformBuffer(false);	// Note: Ribbons don't generate accurate velocities so disabling	
+	MeshElement.PrimitiveUniformBuffer = SceneProxy->GetCustomUniformBuffer(RHICmdList, false);	// Note: Ribbons don't generate accurate velocities so disabling	
 }
 
 void FNiagaraRendererRibbons::InitializeViewIndexBuffersGPU(FRHICommandListImmediate& RHICmdList, FNiagaraGpuComputeDispatchInterface* ComputeDispatchInterface, const FNiagaraRibbonGPUInitParameters& GpuInitParameters, const TSharedPtr<FNiagaraRibbonRenderingFrameViewResources>& RenderingViewResources) const
diff --git a/Engine/Plugins/FX/Niagara/Source/Niagara/Private/NiagaraRendererSprites.cpp b/Engine/Plugins/FX/Niagara/Source/Niagara/Private/NiagaraRendererSprites.cpp
--- a/Engine/Plugins/FX/Niagara/Source/Niagara/Private/NiagaraRendererSprites.cpp
+++ b/Engine/Plugins/FX/Niagara/Source/Niagara/Private/NiagaraRendererSprites.cpp
@@ -929,7 +929,7 @@ void FNiagaraRendererSprites::CreateMeshBatchForView(
 	MeshElement.NumInstances = FMath::Max(0u, NumInstances);
 	MeshElement.MinVertexIndex = 0;
 	MeshElement.MaxVertexIndex = 0;
-	MeshElement.PrimitiveUniformBuffer = SceneProxy.GetCustomUniformBuffer(IsMotionBlurEnabled());
+	MeshElement.PrimitiveUniformBuffer = SceneProxy.GetCustomUniformBuffer(RHICmdList, IsMotionBlurEnabled());
 	if (IndirectDraw.IsValid())
 	{
 		MeshElement.IndirectArgsBuffer = IndirectDraw.Buffer;
@@ -950,7 +950,7 @@ void FNiagaraRendererSprites::GetDynamicMeshElements(const TArray<const FSceneVi
 	check(SceneProxy);
 	PARTICLE_PERF_STAT_CYCLES_RT(SceneProxy->GetProxyDynamicData().PerfStatsContext, GetDynamicMeshElements);
 
-	FRHICommandListBase& RHICmdList = FRHICommandListExecutor::GetImmediateCommandList();
+	FRHICommandListBase& RHICmdList = Collector.GetRHICommandList();
 
 	// Prepare our particle render data
 	// This will also determine if we have anything to render
@@ -1026,7 +1026,7 @@ void FNiagaraRendererSprites::GetDynamicMeshElements(const TArray<const FSceneVi
 				if (ParticleSpriteRenderData.bSortCullOnGpu)
 				{
 					SortInfo.CulledGPUParticleCountOffset = ParticleSpriteRenderData.bNeedsCull ? ComputeDispatchInterface->GetGPUInstanceCounterManager().AcquireCulledEntry() : INDEX_NONE;
-					if (ComputeDispatchInterface->AddSortedGPUSimulation(SortInfo))
+					if (ComputeDispatchInterface->AddSortedGPUSimulation(RHICmdList, SortInfo))
 					{
 						VertexFactory.SetSortedIndices(SortInfo.AllocationInfo.BufferSRV, SortInfo.AllocationInfo.BufferOffset);
 					}
@@ -1077,7 +1077,7 @@ void FNiagaraRendererSprites::GetDynamicRayTracingInstances(FRayTracingMaterialG
 		return;
 	}
 
-	FRHICommandListBase& RHICmdList = FRHICommandListImmediate::Get();
+	FRHICommandListBase& RHICmdList = Context.RHICmdList;
 
 #if STATS
 	FScopeCycleCounter EmitterStatsCounter(EmitterStatID);
@@ -1110,7 +1110,7 @@ void FNiagaraRendererSprites::GetDynamicRayTracingInstances(FRayTracingMaterialG
 		if (ParticleSpriteRenderData.bSortCullOnGpu)
 		{
 			SortInfo.CulledGPUParticleCountOffset = ParticleSpriteRenderData.bNeedsCull ? ComputeDispatchInterface->GetGPUInstanceCounterManager().AcquireCulledEntry() : INDEX_NONE;
-			if (ComputeDispatchInterface->AddSortedGPUSimulation(SortInfo))
+			if (ComputeDispatchInterface->AddSortedGPUSimulation(RHICmdList, SortInfo))
 			{
 				VertexFactory.SetSortedIndices(SortInfo.AllocationInfo.BufferSRV, SortInfo.AllocationInfo.BufferOffset);
 			}
@@ -1126,7 +1126,7 @@ void FNiagaraRendererSprites::GetDynamicRayTracingInstances(FRayTracingMaterialG
 
 	if (NumInstances > 0)
 	{
-		SetupVertexFactory(Context.GraphBuilder.RHICmdList, ParticleSpriteRenderData, VertexFactory);
+		SetupVertexFactory(RHICmdList, ParticleSpriteRenderData, VertexFactory);
 		CollectorResources->UniformBuffer = CreateViewUniformBuffer(ParticleSpriteRenderData, *Context.ReferenceView, Context.ReferenceViewFamily, *SceneProxy, VertexFactory);
 		VertexFactory.SetSpriteUniformBuffer(CollectorResources->UniformBuffer);
 
diff --git a/Engine/Plugins/FX/Niagara/Source/Niagara/Public/NiagaraGpuComputeDispatchInterface.h b/Engine/Plugins/FX/Niagara/Source/Niagara/Public/NiagaraGpuComputeDispatchInterface.h
--- a/Engine/Plugins/FX/Niagara/Source/Niagara/Public/NiagaraGpuComputeDispatchInterface.h
+++ b/Engine/Plugins/FX/Niagara/Source/Niagara/Public/NiagaraGpuComputeDispatchInterface.h
@@ -54,13 +54,15 @@ public:
 	 *
 	 * Return true if the work was registered, or false it GPU sorting is not available or impossible.
 	 */
-	virtual bool AddSortedGPUSimulation(struct FNiagaraGPUSortInfo& SortInfo) = 0;
+	virtual bool AddSortedGPUSimulation(FRHICommandListBase& RHICmdList, struct FNiagaraGPUSortInfo& SortInfo) = 0;
+
+	UE_DEPRECATED(5.4, "AddSortedGPUSimulation requires an RHI command list")
+	virtual bool AddSortedGPUSimulation(struct FNiagaraGPUSortInfo& SortInfo) final { return false; }
 
 	/** Get or create the a data manager, must be done on the rendering thread only. */
 	template<typename TManager>
 	TManager& GetOrCreateDataManager()
-	{	
-		check(IsInRenderingThread());
+	{
 		const FName ManagerName = TManager::GetManagerName();
 		for (auto& DataManager : GpuDataManagers)
 		{
@@ -88,8 +90,8 @@ public:
 	NIAGARA_API virtual const FGlobalDistanceFieldParameterData* GetGlobalDistanceFieldData() const = 0;
 
 	/** Get access to the instance count manager. */
-	FORCEINLINE FNiagaraGPUInstanceCountManager& GetGPUInstanceCounterManager() { check(IsInRenderingThread()); return GPUInstanceCounterManager; }
-	FORCEINLINE const FNiagaraGPUInstanceCountManager& GetGPUInstanceCounterManager() const { check(IsInRenderingThread()); return GPUInstanceCounterManager; }
+	FORCEINLINE FNiagaraGPUInstanceCountManager& GetGPUInstanceCounterManager() { check(IsInParallelRenderingThread()); return GPUInstanceCounterManager; }
+	FORCEINLINE const FNiagaraGPUInstanceCountManager& GetGPUInstanceCounterManager() const { check(IsInParallelRenderingThread()); return GPUInstanceCounterManager; }
 
 #if NIAGARA_COMPUTEDEBUG_ENABLED
 	/** Public interface to Niagara compute debugging. */
diff --git a/Engine/Plugins/FX/Niagara/Source/Niagara/Public/NiagaraRenderableMeshInterface.h b/Engine/Plugins/FX/Niagara/Source/Niagara/Public/NiagaraRenderableMeshInterface.h
--- a/Engine/Plugins/FX/Niagara/Source/Niagara/Public/NiagaraRenderableMeshInterface.h
+++ b/Engine/Plugins/FX/Niagara/Source/Niagara/Public/NiagaraRenderableMeshInterface.h
@@ -45,7 +45,11 @@ public:
 	// Gather all the relevant mesh data to render the mesh
 	virtual void GetLODModelData(FLODModelData& OutLODModelData, int32 LODLevel) const = 0;
 	// Setup the vertex factory for the mesh
-	virtual void SetupVertexFactory(class FNiagaraMeshVertexFactory& InVertexFactory, const FLODModelData& LODModelData) const = 0;
+	virtual void SetupVertexFactory(FRHICommandListBase& RHICmdList, class FNiagaraMeshVertexFactory& InVertexFactory, const FLODModelData& LODModelData) const = 0;
+
+	UE_DEPRECATED(5.4, "SetupVertexFactory requires an RHI command list.")
+	virtual void SetupVertexFactory(class FNiagaraMeshVertexFactory& InVertexFactory, const FLODModelData& LODModelData) const final {}
+
 	// Gather a list of used materials
 	virtual void GetUsedMaterials(TArray<UMaterialInterface*>& OutMaterials) const = 0;
 
diff --git a/Engine/Plugins/FX/Niagara/Source/Niagara/Public/NiagaraRendererProperties.h b/Engine/Plugins/FX/Niagara/Source/Niagara/Public/NiagaraRendererProperties.h
--- a/Engine/Plugins/FX/Niagara/Source/Niagara/Public/NiagaraRendererProperties.h
+++ b/Engine/Plugins/FX/Niagara/Source/Niagara/Public/NiagaraRendererProperties.h
@@ -146,9 +146,9 @@ struct FNiagaraRendererLayout
 	NIAGARA_API void Finalize();
 
 	TConstArrayView<FNiagaraRendererVariableInfo> GetVFVariables_GameThread() const { check(IsInGameThread() || IsInParallelGameThread()); return MakeArrayView(VFVariables_GT); }
-	TConstArrayView<FNiagaraRendererVariableInfo> GetVFVariables_RenderThread() const { check(IsInRenderingThread()); return MakeArrayView(VFVariables_RT); }
-	int32 GetTotalFloatComponents_RenderThread() const { check(IsInRenderingThread()); return TotalFloatComponents_RT; }
-	int32 GetTotalHalfComponents_RenderThread() const { check(IsInRenderingThread()); return TotalHalfComponents_RT; }
+	TConstArrayView<FNiagaraRendererVariableInfo> GetVFVariables_RenderThread() const { check(IsInParallelRenderingThread()); return MakeArrayView(VFVariables_RT); }
+	int32 GetTotalFloatComponents_RenderThread() const { check(IsInParallelRenderingThread()); return TotalFloatComponents_RT; }
+	int32 GetTotalHalfComponents_RenderThread() const { check(IsInParallelRenderingThread()); return TotalHalfComponents_RT; }
 
 	SIZE_T GetAllocatedSize() const { return VFVariables_GT.GetAllocatedSize() + VFVariables_RT.GetAllocatedSize(); }
 
diff --git a/Engine/Plugins/FX/Niagara/Source/Niagara/Public/NiagaraSceneProxy.h b/Engine/Plugins/FX/Niagara/Source/Niagara/Public/NiagaraSceneProxy.h
--- a/Engine/Plugins/FX/Niagara/Source/Niagara/Public/NiagaraSceneProxy.h
+++ b/Engine/Plugins/FX/Niagara/Source/Niagara/Public/NiagaraSceneProxy.h
@@ -48,7 +48,14 @@ public:
 
 	NIAGARA_API const FVector3f& GetLWCRenderTile() const;
 
+	NIAGARA_API TUniformBuffer<FPrimitiveUniformShaderParameters>* GetCustomUniformBufferResource(FRHICommandListBase& RHICmdList, bool bHasVelocity, const FBox& InstanceBounds = FBox(ForceInitToZero)) const;
+
+	UE_DEPRECATED(5.4, "GetCustomUniformBufferResource requires a command list.")
 	NIAGARA_API TUniformBuffer<FPrimitiveUniformShaderParameters>* GetCustomUniformBufferResource(bool bHasVelocity, const FBox& InstanceBounds = FBox(ForceInitToZero)) const;
+
+	NIAGARA_API FRHIUniformBuffer* GetCustomUniformBuffer(FRHICommandListBase& RHICmdList, bool bHasVelocity, const FBox& InstanceBounds = FBox(ForceInitToZero)) const;
+
+	UE_DEPRECATED(5.4, "GetCustomUniformBuffer requires a command list.")
 	NIAGARA_API FRHIUniformBuffer* GetCustomUniformBuffer(bool bHasVelocity, const FBox& InstanceBounds = FBox(ForceInitToZero)) const;
 
 	NIAGARA_API virtual FPrimitiveViewRelevance GetViewRelevance(const FSceneView* View) const override;
diff --git a/Engine/Plugins/FX/Niagara/Source/NiagaraVertexFactories/Private/NiagaraMeshVertexFactory.cpp b/Engine/Plugins/FX/Niagara/Source/NiagaraVertexFactories/Private/NiagaraMeshVertexFactory.cpp
--- a/Engine/Plugins/FX/Niagara/Source/NiagaraVertexFactories/Private/NiagaraMeshVertexFactory.cpp
+++ b/Engine/Plugins/FX/Niagara/Source/NiagaraVertexFactories/Private/NiagaraMeshVertexFactory.cpp
@@ -220,10 +220,10 @@ void FNiagaraMeshVertexFactory::GetVertexElements(ERHIFeatureLevel::Type Feature
 #endif
 }
 
-void FNiagaraMeshVertexFactory::SetData(const FStaticMeshDataType& InData)
+void FNiagaraMeshVertexFactory::SetData(FRHICommandListBase& RHICmdList, const FStaticMeshDataType& InData)
 {
 	Data = InData;
-	UpdateRHI(FRHICommandListImmediate::Get());
+	UpdateRHI(RHICmdList);
 }
 
 #if NIAGARA_ENABLE_GPU_SCENE_MESHES
diff --git a/Engine/Plugins/FX/Niagara/Source/NiagaraVertexFactories/Public/NiagaraMeshVertexFactory.h b/Engine/Plugins/FX/Niagara/Source/NiagaraVertexFactories/Public/NiagaraMeshVertexFactory.h
--- a/Engine/Plugins/FX/Niagara/Source/NiagaraVertexFactories/Public/NiagaraMeshVertexFactory.h
+++ b/Engine/Plugins/FX/Niagara/Source/NiagaraVertexFactories/Public/NiagaraMeshVertexFactory.h
@@ -158,7 +158,7 @@ public:
 	/**
 	* An implementation of the interface used by TSynchronizedResource to update the resource with new data from the game thread.
 	*/
-	NIAGARAVERTEXFACTORIES_API void SetData(const FStaticMeshDataType& InData);
+	NIAGARAVERTEXFACTORIES_API void SetData(FRHICommandListBase& RHICmdList, const FStaticMeshDataType& InData);
 
 	/**
 	* Set the uniform buffer for this vertex factory.
diff --git a/Engine/Plugins/Runtime/CableComponent/Source/CableComponent/Private/CableComponent.cpp b/Engine/Plugins/Runtime/CableComponent/Source/CableComponent/Private/CableComponent.cpp
--- a/Engine/Plugins/Runtime/CableComponent/Source/CableComponent/Private/CableComponent.cpp
+++ b/Engine/Plugins/Runtime/CableComponent/Source/CableComponent/Private/CableComponent.cpp
@@ -344,7 +344,7 @@ public:
 				bOutputVelocity |= AlwaysHasVelocity();
 
 				FDynamicPrimitiveUniformBuffer& DynamicPrimitiveUniformBuffer = Collector.AllocateOneFrameResource<FDynamicPrimitiveUniformBuffer>();
-				DynamicPrimitiveUniformBuffer.Set(GetLocalToWorld(), PreviousLocalToWorld, GetBounds(), GetLocalBounds(), true, bHasPrecomputedVolumetricLightmap, bOutputVelocity);
+				DynamicPrimitiveUniformBuffer.Set(Collector.GetRHICommandList(), GetLocalToWorld(), PreviousLocalToWorld, GetBounds(), GetLocalBounds(), true, bHasPrecomputedVolumetricLightmap, bOutputVelocity);
 				BatchElement.PrimitiveUniformBufferResource = &DynamicPrimitiveUniformBuffer.UniformBuffer;
 
 				BatchElement.FirstIndex = 0;
@@ -483,7 +483,7 @@ public:
 			bOutputVelocity |= AlwaysHasVelocity();
 
 			FDynamicPrimitiveUniformBuffer& DynamicPrimitiveUniformBuffer = Context.RayTracingMeshResourceCollector.AllocateOneFrameResource<FDynamicPrimitiveUniformBuffer>();
-			DynamicPrimitiveUniformBuffer.Set(GetLocalToWorld(), PreviousLocalToWorld, GetBounds(), GetLocalBounds(), ReceivesDecals(), bHasPrecomputedVolumetricLightmap, bOutputVelocity);
+			DynamicPrimitiveUniformBuffer.Set(Context.RHICmdList, GetLocalToWorld(), PreviousLocalToWorld, GetBounds(), GetLocalBounds(), ReceivesDecals(), bHasPrecomputedVolumetricLightmap, bOutputVelocity);
 			BatchElement.PrimitiveUniformBufferResource = &DynamicPrimitiveUniformBuffer.UniformBuffer;
 			BatchElement.FirstIndex = 0;
 			BatchElement.NumPrimitives = GetRequiredIndexCount() / 3;
diff --git a/Engine/Plugins/Runtime/CustomMeshComponent/Source/CustomMeshComponent/Private/CustomMeshComponent.cpp b/Engine/Plugins/Runtime/CustomMeshComponent/Source/CustomMeshComponent/Private/CustomMeshComponent.cpp
--- a/Engine/Plugins/Runtime/CustomMeshComponent/Source/CustomMeshComponent/Private/CustomMeshComponent.cpp
+++ b/Engine/Plugins/Runtime/CustomMeshComponent/Source/CustomMeshComponent/Private/CustomMeshComponent.cpp
@@ -134,7 +134,7 @@ public:
 				bOutputVelocity |= AlwaysHasVelocity();
 
 				FDynamicPrimitiveUniformBuffer& DynamicPrimitiveUniformBuffer = Collector.AllocateOneFrameResource<FDynamicPrimitiveUniformBuffer>();
-				DynamicPrimitiveUniformBuffer.Set(GetLocalToWorld(), PreviousLocalToWorld, GetBounds(), GetLocalBounds(), true, bHasPrecomputedVolumetricLightmap, bOutputVelocity);
+				DynamicPrimitiveUniformBuffer.Set(Collector.GetRHICommandList(), GetLocalToWorld(), PreviousLocalToWorld, GetBounds(), GetLocalBounds(), true, bHasPrecomputedVolumetricLightmap, bOutputVelocity);
 				BatchElement.PrimitiveUniformBufferResource = &DynamicPrimitiveUniformBuffer.UniformBuffer;
 
 				BatchElement.FirstIndex = 0;
diff --git a/Engine/Plugins/Runtime/GeometryCache/Source/GeometryCache/Private/GeometryCacheSceneProxy.cpp b/Engine/Plugins/Runtime/GeometryCache/Source/GeometryCache/Private/GeometryCacheSceneProxy.cpp
--- a/Engine/Plugins/Runtime/GeometryCache/Source/GeometryCache/Private/GeometryCacheSceneProxy.cpp
+++ b/Engine/Plugins/Runtime/GeometryCache/Source/GeometryCache/Private/GeometryCacheSceneProxy.cpp
@@ -337,7 +337,7 @@ void FGeometryCacheSceneProxy::CreateMeshBatch(
 	GetScene().GetPrimitiveUniformShaderParameters_RenderThread(GetPrimitiveSceneInfo(), bHasPrecomputedVolumetricLightmap, PreviousLocalToWorld, SingleCaptureIndex, bOutputVelocity);
 	bOutputVelocity |= AlwaysHasVelocity();
 
-	DynamicPrimitiveUniformBuffer.Set(LocalToWorldTransform, PreviousLocalToWorld, GetBounds(), GetLocalBounds(), ReceivesDecals(), false, bOutputVelocity);
+	DynamicPrimitiveUniformBuffer.Set(RHICmdList, LocalToWorldTransform, PreviousLocalToWorld, GetBounds(), GetLocalBounds(), ReceivesDecals(), false, bOutputVelocity);
 	BatchElement.PrimitiveUniformBufferResource = &DynamicPrimitiveUniformBuffer.UniformBuffer;
 
 	const FGeometryCacheMeshData* MeshData = TrackProxy->bNextFrameMeshDataSelected ? TrackProxy->NextFrameMeshData : TrackProxy->MeshData;
@@ -401,7 +401,7 @@ HHitProxy* FGeometryCacheSceneProxy::CreateHitProxies(UPrimitiveComponent* Compo
 void FGeometryCacheSceneProxy::GetDynamicMeshElements(const TArray<const FSceneView*>& Views, const FSceneViewFamily& ViewFamily, uint32 VisibilityMap, FMeshElementCollector& Collector) const
 {
 	SCOPE_CYCLE_COUNTER(STAT_GeometryCacheSceneProxy_GetMeshElements);
-	FRHICommandListBase& RHICmdList = FRHICommandListImmediate::Get();
+	FRHICommandListBase& RHICmdList = Collector.GetRHICommandList();
 
 	// Set up wire frame material (if needed)
 	const bool bWireframe = AllowDebugViewmodes() && ViewFamily.EngineShowFlags.Wireframe;
diff --git a/Engine/Plugins/Runtime/HairStrands/Source/HairStrandsCore/Private/GroomComponent.cpp b/Engine/Plugins/Runtime/HairStrands/Source/HairStrandsCore/Private/GroomComponent.cpp
--- a/Engine/Plugins/Runtime/HairStrands/Source/HairStrandsCore/Private/GroomComponent.cpp
+++ b/Engine/Plugins/Runtime/HairStrands/Source/HairStrandsCore/Private/GroomComponent.cpp
@@ -978,7 +978,7 @@ public:
 		}
 
 		FDynamicPrimitiveUniformBuffer& DynamicPrimitiveUniformBuffer = Collector.AllocateOneFrameResource<FDynamicPrimitiveUniformBuffer>();
-		DynamicPrimitiveUniformBuffer.Set(CurrentLocalToWorld, PreviousLocalToWorld, GetBounds(), GetLocalBounds(), true, false, bOutputVelocity);
+		DynamicPrimitiveUniformBuffer.Set(Collector.GetRHICommandList(), CurrentLocalToWorld, PreviousLocalToWorld, GetBounds(), GetLocalBounds(), true, false, bOutputVelocity);
 		BatchElement.PrimitiveUniformBufferResource = &DynamicPrimitiveUniformBuffer.UniformBuffer; // automatic copy to the gpu scene buffer
 		//primtiveid is set to 0
 
diff --git a/Engine/Plugins/Runtime/MeshModelingToolset/Source/ModelingComponents/Private/Drawing/LineSetComponent.cpp b/Engine/Plugins/Runtime/MeshModelingToolset/Source/ModelingComponents/Private/Drawing/LineSetComponent.cpp
--- a/Engine/Plugins/Runtime/MeshModelingToolset/Source/ModelingComponents/Private/Drawing/LineSetComponent.cpp
+++ b/Engine/Plugins/Runtime/MeshModelingToolset/Source/ModelingComponents/Private/Drawing/LineSetComponent.cpp
@@ -161,7 +161,7 @@ public:
 					Mesh.MaterialRenderProxy = MeshBatchData.MaterialProxy;
 
 					FDynamicPrimitiveUniformBuffer& DynamicPrimitiveUniformBuffer = Collector.AllocateOneFrameResource<FDynamicPrimitiveUniformBuffer>();
-					DynamicPrimitiveUniformBuffer.Set(GetLocalToWorld(), GetLocalToWorld(), GetBounds(), GetLocalBounds(), false, false, AlwaysHasVelocity());
+					DynamicPrimitiveUniformBuffer.Set(Collector.GetRHICommandList(), GetLocalToWorld(), GetLocalToWorld(), GetBounds(), GetLocalBounds(), false, false, AlwaysHasVelocity());
 					BatchElement.PrimitiveUniformBufferResource = &DynamicPrimitiveUniformBuffer.UniformBuffer;
 
 					BatchElement.FirstIndex = MeshBatchData.StartIndex;
diff --git a/Engine/Plugins/Runtime/MeshModelingToolset/Source/ModelingComponents/Private/Drawing/MeshWireframeComponent.cpp b/Engine/Plugins/Runtime/MeshModelingToolset/Source/ModelingComponents/Private/Drawing/MeshWireframeComponent.cpp
--- a/Engine/Plugins/Runtime/MeshModelingToolset/Source/ModelingComponents/Private/Drawing/MeshWireframeComponent.cpp
+++ b/Engine/Plugins/Runtime/MeshModelingToolset/Source/ModelingComponents/Private/Drawing/MeshWireframeComponent.cpp
@@ -287,7 +287,7 @@ public:
 					Mesh.MaterialRenderProxy = MeshBatchData.MaterialProxy;
 
 					FDynamicPrimitiveUniformBuffer& DynamicPrimitiveUniformBuffer = Collector.AllocateOneFrameResource<FDynamicPrimitiveUniformBuffer>();
-					DynamicPrimitiveUniformBuffer.Set(GetLocalToWorld(), GetLocalToWorld(), GetBounds(), GetLocalBounds(), false, false, AlwaysHasVelocity());
+					DynamicPrimitiveUniformBuffer.Set(Collector.GetRHICommandList(), GetLocalToWorld(), GetLocalToWorld(), GetBounds(), GetLocalBounds(), false, false, AlwaysHasVelocity());
 					BatchElement.PrimitiveUniformBufferResource = &DynamicPrimitiveUniformBuffer.UniformBuffer;
 
 					BatchElement.FirstIndex = MeshBatchData.StartIndex;
diff --git a/Engine/Plugins/Runtime/MeshModelingToolset/Source/ModelingComponents/Private/Drawing/PointSetComponent.cpp b/Engine/Plugins/Runtime/MeshModelingToolset/Source/ModelingComponents/Private/Drawing/PointSetComponent.cpp
--- a/Engine/Plugins/Runtime/MeshModelingToolset/Source/ModelingComponents/Private/Drawing/PointSetComponent.cpp
+++ b/Engine/Plugins/Runtime/MeshModelingToolset/Source/ModelingComponents/Private/Drawing/PointSetComponent.cpp
@@ -175,7 +175,7 @@ public:
 					Mesh.MaterialRenderProxy = MeshBatchData.MaterialProxy;
 
 					FDynamicPrimitiveUniformBuffer& DynamicPrimitiveUniformBuffer = Collector.AllocateOneFrameResource<FDynamicPrimitiveUniformBuffer>();
-					DynamicPrimitiveUniformBuffer.Set(GetLocalToWorld(), GetLocalToWorld(), GetBounds(), GetLocalBounds(), true, false, AlwaysHasVelocity());
+					DynamicPrimitiveUniformBuffer.Set(Collector.GetRHICommandList(), GetLocalToWorld(), GetLocalToWorld(), GetBounds(), GetLocalBounds(), true, false, AlwaysHasVelocity());
 					BatchElement.PrimitiveUniformBufferResource = &DynamicPrimitiveUniformBuffer.UniformBuffer;
 
 					BatchElement.FirstIndex = MeshBatchData.StartIndex;
diff --git a/Engine/Plugins/Runtime/MeshModelingToolset/Source/ModelingComponents/Private/Drawing/TriangleSetComponent.cpp b/Engine/Plugins/Runtime/MeshModelingToolset/Source/ModelingComponents/Private/Drawing/TriangleSetComponent.cpp
--- a/Engine/Plugins/Runtime/MeshModelingToolset/Source/ModelingComponents/Private/Drawing/TriangleSetComponent.cpp
+++ b/Engine/Plugins/Runtime/MeshModelingToolset/Source/ModelingComponents/Private/Drawing/TriangleSetComponent.cpp
@@ -145,7 +145,7 @@ public:
 					Mesh.MaterialRenderProxy = MeshBatchData.MaterialProxy;
 
 					FDynamicPrimitiveUniformBuffer& DynamicPrimitiveUniformBuffer = Collector.AllocateOneFrameResource<FDynamicPrimitiveUniformBuffer>();
-					DynamicPrimitiveUniformBuffer.Set(GetLocalToWorld(), GetLocalToWorld(), GetBounds(), GetLocalBounds(), true, false, AlwaysHasVelocity());
+					DynamicPrimitiveUniformBuffer.Set(Collector.GetRHICommandList(), GetLocalToWorld(), GetLocalToWorld(), GetBounds(), GetLocalBounds(), true, false, AlwaysHasVelocity());
 					BatchElement.PrimitiveUniformBufferResource = &DynamicPrimitiveUniformBuffer.UniformBuffer;
 
 					BatchElement.FirstIndex = MeshBatchData.StartIndex;
diff --git a/Engine/Plugins/Runtime/ProceduralMeshComponent/Source/ProceduralMeshComponent/Private/ProceduralMeshComponent.cpp b/Engine/Plugins/Runtime/ProceduralMeshComponent/Source/ProceduralMeshComponent/Private/ProceduralMeshComponent.cpp
--- a/Engine/Plugins/Runtime/ProceduralMeshComponent/Source/ProceduralMeshComponent/Private/ProceduralMeshComponent.cpp
+++ b/Engine/Plugins/Runtime/ProceduralMeshComponent/Source/ProceduralMeshComponent/Private/ProceduralMeshComponent.cpp
@@ -377,7 +377,7 @@ public:
 						bOutputVelocity |= AlwaysHasVelocity();
 
 						FDynamicPrimitiveUniformBuffer& DynamicPrimitiveUniformBuffer = Collector.AllocateOneFrameResource<FDynamicPrimitiveUniformBuffer>();
-						DynamicPrimitiveUniformBuffer.Set(GetLocalToWorld(), PreviousLocalToWorld, GetBounds(), GetLocalBounds(), GetLocalBounds(), ReceivesDecals(), bHasPrecomputedVolumetricLightmap, bOutputVelocity, GetCustomPrimitiveData());
+						DynamicPrimitiveUniformBuffer.Set(Collector.GetRHICommandList(), GetLocalToWorld(), PreviousLocalToWorld, GetBounds(), GetLocalBounds(), GetLocalBounds(), ReceivesDecals(), bHasPrecomputedVolumetricLightmap, bOutputVelocity, GetCustomPrimitiveData());
 						BatchElement.PrimitiveUniformBufferResource = &DynamicPrimitiveUniformBuffer.UniformBuffer;
 
 						BatchElement.FirstIndex = 0;
@@ -495,7 +495,7 @@ public:
 					bOutputVelocity |= AlwaysHasVelocity();
 
 					FDynamicPrimitiveUniformBuffer& DynamicPrimitiveUniformBuffer = Context.RayTracingMeshResourceCollector.AllocateOneFrameResource<FDynamicPrimitiveUniformBuffer>();
-					DynamicPrimitiveUniformBuffer.Set(GetLocalToWorld(), PreviousLocalToWorld, GetBounds(), GetLocalBounds(), GetLocalBounds(), ReceivesDecals(), bHasPrecomputedVolumetricLightmap, bOutputVelocity, GetCustomPrimitiveData());
+					DynamicPrimitiveUniformBuffer.Set(Context.RHICmdList, GetLocalToWorld(), PreviousLocalToWorld, GetBounds(), GetLocalBounds(), GetLocalBounds(), ReceivesDecals(), bHasPrecomputedVolumetricLightmap, bOutputVelocity, GetCustomPrimitiveData());
 					BatchElement.PrimitiveUniformBufferResource = &DynamicPrimitiveUniformBuffer.UniformBuffer;
 
 					BatchElement.FirstIndex = 0;
diff --git a/Engine/Plugins/Runtime/nDisplay/Source/DisplayCluster/Private/Game/EngineClasses/Scene/DisplayClusterWidgetComponent.cpp b/Engine/Plugins/Runtime/nDisplay/Source/DisplayCluster/Private/Game/EngineClasses/Scene/DisplayClusterWidgetComponent.cpp
--- a/Engine/Plugins/Runtime/nDisplay/Source/DisplayCluster/Private/Game/EngineClasses/Scene/DisplayClusterWidgetComponent.cpp
+++ b/Engine/Plugins/Runtime/nDisplay/Source/DisplayCluster/Private/Game/EngineClasses/Scene/DisplayClusterWidgetComponent.cpp
@@ -172,23 +172,20 @@ public:
 			}
 
 			FLocalVertexFactory* VertexFactoryPtr = &VertexFactory;
-			ENQUEUE_RENDER_COMMAND(FWidgetBillboardProxyLegacyInit)(
-				[VertexFactoryPtr, this](FRHICommandListImmediate& RHICmdList)
-			{
-				this->StaticMeshVertexBuffers.PositionVertexBuffer.UpdateRHI(RHICmdList);
-				this->StaticMeshVertexBuffers.StaticMeshVertexBuffer.UpdateRHI(RHICmdList);
-				this->StaticMeshVertexBuffers.ColorVertexBuffer.UpdateRHI(RHICmdList);
-
-				FLocalVertexFactory::FDataType Data;
-				this->StaticMeshVertexBuffers.PositionVertexBuffer.BindPositionVertexBuffer(VertexFactoryPtr, Data);
-				this->StaticMeshVertexBuffers.StaticMeshVertexBuffer.BindTangentVertexBuffer(VertexFactoryPtr, Data);
-				this->StaticMeshVertexBuffers.StaticMeshVertexBuffer.BindPackedTexCoordVertexBuffer(VertexFactoryPtr, Data);
-				this->StaticMeshVertexBuffers.StaticMeshVertexBuffer.BindLightMapVertexBuffer(VertexFactoryPtr, Data, 0);
-				this->StaticMeshVertexBuffers.ColorVertexBuffer.BindColorVertexBuffer(VertexFactoryPtr, Data);
-				VertexFactoryPtr->SetData(RHICmdList, Data);
-
-				VertexFactoryPtr->UpdateRHI(RHICmdList);
-			});
+			FRHICommandListBase& RHICmdList = Collector.GetRHICommandList();
+
+			StaticMeshVertexBuffers.PositionVertexBuffer.UpdateRHI(RHICmdList);
+			StaticMeshVertexBuffers.StaticMeshVertexBuffer.UpdateRHI(RHICmdList);
+			StaticMeshVertexBuffers.ColorVertexBuffer.UpdateRHI(RHICmdList);
+
+			FLocalVertexFactory::FDataType Data;
+			StaticMeshVertexBuffers.PositionVertexBuffer.BindPositionVertexBuffer(VertexFactoryPtr, Data);
+			StaticMeshVertexBuffers.StaticMeshVertexBuffer.BindTangentVertexBuffer(VertexFactoryPtr, Data);
+			StaticMeshVertexBuffers.StaticMeshVertexBuffer.BindPackedTexCoordVertexBuffer(VertexFactoryPtr, Data);
+			StaticMeshVertexBuffers.StaticMeshVertexBuffer.BindLightMapVertexBuffer(VertexFactoryPtr, Data, 0);
+			StaticMeshVertexBuffers.ColorVertexBuffer.BindColorVertexBuffer(VertexFactoryPtr, Data);
+			VertexFactoryPtr->SetData(RHICmdList, Data);
+			VertexFactoryPtr->UpdateRHI(RHICmdList);
 		}
 	}
 
diff --git a/Engine/Source/Editor/UnrealEd/Private/Animation/DebugSkelMeshComponent.cpp b/Engine/Source/Editor/UnrealEd/Private/Animation/DebugSkelMeshComponent.cpp
--- a/Engine/Source/Editor/UnrealEd/Private/Animation/DebugSkelMeshComponent.cpp
+++ b/Engine/Source/Editor/UnrealEd/Private/Animation/DebugSkelMeshComponent.cpp
@@ -1271,6 +1271,8 @@ SIZE_T FDebugSkelMeshSceneProxy::GetTypeHash() const
 
 void FDebugSkelMeshSceneProxy::GetDynamicMeshElements(const TArray<const FSceneView*>& Views, const FSceneViewFamily& ViewFamily, uint32 VisibilityMap, FMeshElementCollector& Collector) const
 {
+	FRHICommandListBase& RHICmdList = Collector.GetRHICommandList();
+
 	if(!DynamicData || DynamicData->bDrawMesh)
 	{
 		GetMeshElementsConditionallySelectable(Views, ViewFamily, bSelectable, VisibilityMap, Collector);
diff --git a/Engine/Source/Runtime/Engine/Private/Components/ArrowComponent.cpp b/Engine/Source/Runtime/Engine/Private/Components/ArrowComponent.cpp
--- a/Engine/Source/Runtime/Engine/Private/Components/ArrowComponent.cpp
+++ b/Engine/Source/Runtime/Engine/Private/Components/ArrowComponent.cpp
@@ -147,7 +147,7 @@ public:
 				Mesh.MaterialRenderProxy = ArrowMaterialRenderProxy;
 
 				FDynamicPrimitiveUniformBuffer& DynamicPrimitiveUniformBuffer = Collector.AllocateOneFrameResource<FDynamicPrimitiveUniformBuffer>();
-				DynamicPrimitiveUniformBuffer.Set(FScaleMatrix(ViewScale) * EffectiveLocalToWorld, FScaleMatrix(ViewScale) * EffectiveLocalToWorld, GetBounds(), GetLocalBounds(), true, false, AlwaysHasVelocity());
+				DynamicPrimitiveUniformBuffer.Set(Collector.GetRHICommandList(), FScaleMatrix(ViewScale) * EffectiveLocalToWorld, FScaleMatrix(ViewScale) * EffectiveLocalToWorld, GetBounds(), GetLocalBounds(), true, false, AlwaysHasVelocity());
 				BatchElement.PrimitiveUniformBufferResource = &DynamicPrimitiveUniformBuffer.UniformBuffer;
 
 				BatchElement.FirstIndex = 0;
diff --git a/Engine/Source/Runtime/Engine/Private/Components/HeterogeneousVolumeComponent.cpp b/Engine/Source/Runtime/Engine/Private/Components/HeterogeneousVolumeComponent.cpp
--- a/Engine/Source/Runtime/Engine/Private/Components/HeterogeneousVolumeComponent.cpp
+++ b/Engine/Source/Runtime/Engine/Private/Components/HeterogeneousVolumeComponent.cpp
@@ -164,8 +164,6 @@ void FHeterogeneousVolumeSceneProxy::GetDynamicMeshElements(
 	uint32 VisibilityMap,
 	FMeshElementCollector& Collector) const
 {
-	check(IsInRenderingThread());
-
 	// Create a dummy MeshBatch to make the system happy..
 	if (MaterialInterface)
 	{
diff --git a/Engine/Source/Runtime/Engine/Private/Components/MaterialBillboardComponent.cpp b/Engine/Source/Runtime/Engine/Private/Components/MaterialBillboardComponent.cpp
--- a/Engine/Source/Runtime/Engine/Private/Components/MaterialBillboardComponent.cpp
+++ b/Engine/Source/Runtime/Engine/Private/Components/MaterialBillboardComponent.cpp
@@ -238,22 +238,19 @@ public:
 			}
 
 			FLocalVertexFactory* VertexFactoryPtr = &VertexFactory;
-			const FMaterialSpriteSceneProxy* Self = this;
-			ENQUEUE_RENDER_COMMAND(FMaterialSpriteSceneProxyLegacyInit)(
-				[VertexFactoryPtr, Self](FRHICommandListImmediate& RHICmdList)
-			{
-				Self->StaticMeshVertexBuffers.PositionVertexBuffer.UpdateRHI(RHICmdList);
-				Self->StaticMeshVertexBuffers.StaticMeshVertexBuffer.UpdateRHI(RHICmdList);
-				Self->StaticMeshVertexBuffers.ColorVertexBuffer.UpdateRHI(RHICmdList);
-
-				FLocalVertexFactory::FDataType Data;
-				Self->StaticMeshVertexBuffers.PositionVertexBuffer.BindPositionVertexBuffer(VertexFactoryPtr, Data);
-				Self->StaticMeshVertexBuffers.StaticMeshVertexBuffer.BindTangentVertexBuffer(VertexFactoryPtr, Data);
-				Self->StaticMeshVertexBuffers.StaticMeshVertexBuffer.BindPackedTexCoordVertexBuffer(VertexFactoryPtr, Data);
-				Self->StaticMeshVertexBuffers.StaticMeshVertexBuffer.BindLightMapVertexBuffer(VertexFactoryPtr, Data, 0);
-				Self->StaticMeshVertexBuffers.ColorVertexBuffer.BindColorVertexBuffer(VertexFactoryPtr, Data);
-				VertexFactoryPtr->SetData(RHICmdList, Data);
-			});
+			FRHICommandListBase& RHICmdList = Collector.GetRHICommandList();
+
+			StaticMeshVertexBuffers.PositionVertexBuffer.UpdateRHI(RHICmdList);
+			StaticMeshVertexBuffers.StaticMeshVertexBuffer.UpdateRHI(RHICmdList);
+			StaticMeshVertexBuffers.ColorVertexBuffer.UpdateRHI(RHICmdList);
+
+			FLocalVertexFactory::FDataType Data;
+			StaticMeshVertexBuffers.PositionVertexBuffer.BindPositionVertexBuffer(VertexFactoryPtr, Data);
+			StaticMeshVertexBuffers.StaticMeshVertexBuffer.BindTangentVertexBuffer(VertexFactoryPtr, Data);
+			StaticMeshVertexBuffers.StaticMeshVertexBuffer.BindPackedTexCoordVertexBuffer(VertexFactoryPtr, Data);
+			StaticMeshVertexBuffers.StaticMeshVertexBuffer.BindLightMapVertexBuffer(VertexFactoryPtr, Data, 0);
+			StaticMeshVertexBuffers.ColorVertexBuffer.BindColorVertexBuffer(VertexFactoryPtr, Data);
+			VertexFactoryPtr->SetData(RHICmdList, Data);
 		}
 	}
 
diff --git a/Engine/Source/Runtime/Engine/Private/GPUSortManager.cpp b/Engine/Source/Runtime/Engine/Private/GPUSortManager.cpp
--- a/Engine/Source/Runtime/Engine/Private/GPUSortManager.cpp
+++ b/Engine/Source/Runtime/Engine/Private/GPUSortManager.cpp
@@ -582,15 +582,13 @@ void FGPUSortManager::Register(const FGPUSortKeyGenDelegate& CallbackDelegate, E
 	Callback.Name = InName;
 }
 
-bool FGPUSortManager::AddTask(FAllocationInfo& OutInfo, int32 ValueCount, EGPUSortFlags TaskFlags)
+bool FGPUSortManager::AddTask(FRHICommandListBase& RHICmdList, FAllocationInfo& OutInfo, int32 ValueCount, EGPUSortFlags TaskFlags)
 {
 	if (!FXConsoleVariables::bAllowGPUSorting)
 	{
 		return false;
 	}
 
-	FRHICommandListBase& RHICmdList = FRHICommandListImmediate::Get();
-
 	LLM_SCOPE(ELLMTag::GPUSort);
 
 	// Can only use flags that were registered initially.
@@ -641,6 +639,11 @@ bool FGPUSortManager::AddTask(FAllocationInfo& OutInfo, int32 ValueCount, EGPUSo
 	return true;
 }
 
+bool FGPUSortManager::AddTask(FAllocationInfo& OutInfo, int32 ValueCount, EGPUSortFlags TaskFlags)
+{
+	return AddTask(FRHICommandListImmediate::Get(), OutInfo, ValueCount, TaskFlags);
+}
+
 void FGPUSortManager::FinalizeSortBatches()
 {
 	// At this point the batches are complete as no more tasks will be added to them.
diff --git a/Engine/Source/Runtime/Engine/Private/MeshParticleVertexFactory.cpp b/Engine/Source/Runtime/Engine/Private/MeshParticleVertexFactory.cpp
--- a/Engine/Source/Runtime/Engine/Private/MeshParticleVertexFactory.cpp
+++ b/Engine/Source/Runtime/Engine/Private/MeshParticleVertexFactory.cpp
@@ -337,10 +337,15 @@ void FMeshParticleVertexFactory::GetPSOPrecacheVertexFetchElements(EVertexInputS
 	}
 }
 
-void FMeshParticleVertexFactory::SetData(const FDataType& InData)
+void FMeshParticleVertexFactory::SetData(FRHICommandListBase& RHICmdList, const FDataType& InData)
 {
 	Data = InData;
-	UpdateRHI(FRHICommandListImmediate::Get());
+	UpdateRHI(RHICmdList);
+}
+
+void FMeshParticleVertexFactory::SetData(const FDataType& InData)
+{
+	SetData(FRHICommandListImmediate::Get(), InData);
 }
 
 IMPLEMENT_VERTEX_FACTORY_PARAMETER_TYPE(FMeshParticleVertexFactory, SF_Vertex, FMeshParticleVertexFactoryShaderParameters);
diff --git a/Engine/Source/Runtime/Engine/Private/Particles/FXSystemPrivate.h b/Engine/Source/Runtime/Engine/Private/Particles/FXSystemPrivate.h
--- a/Engine/Source/Runtime/Engine/Private/Particles/FXSystemPrivate.h
+++ b/Engine/Source/Runtime/Engine/Private/Particles/FXSystemPrivate.h
@@ -261,7 +261,7 @@ public:
 	 * @param OutInfo The bindings for this GPU sort task, if success. 
 	 * @returns true if the work was registered, or false it GPU sorting is not available or impossible.
 	 */
-	bool AddSortedGPUSimulation(FParticleSimulationGPU* Simulation, const FVector& ViewOrigin, bool bIsTranslucent, FGPUSortManager::FAllocationInfo& OutInfo);
+	bool AddSortedGPUSimulation(FRHICommandListBase& RHICmdList, FParticleSimulationGPU* Simulation, const FVector& ViewOrigin, bool bIsTranslucent, FGPUSortManager::FAllocationInfo& OutInfo);
 
 	void PrepareGPUSimulation(FRHICommandListImmediate& RHICmdList);
 	void FinalizeGPUSimulation(FRHICommandListImmediate& RHICmdList);
diff --git a/Engine/Source/Runtime/Engine/Private/Particles/ParticleGpuSimulation.cpp b/Engine/Source/Runtime/Engine/Private/Particles/ParticleGpuSimulation.cpp
--- a/Engine/Source/Runtime/Engine/Private/Particles/ParticleGpuSimulation.cpp
+++ b/Engine/Source/Runtime/Engine/Private/Particles/ParticleGpuSimulation.cpp
@@ -2662,7 +2662,7 @@ public:
 	{
 		auto FeatureLevel = ViewFamily.GetFeatureLevel();
 
-		FRHICommandListBase& RHICmdList = FRHICommandListImmediate::Get();
+		FRHICommandListBase& RHICmdList = Collector.GetRHICommandList();
 
 		if (RHISupportsGPUParticles() && GFXCascadeGpuSpriteRenderingEnabled)
 		{
@@ -2708,7 +2708,7 @@ public:
 					// accumulate all GPU particle emitters that need to be sorted. That is so they can be sorted in one big radix
 					// sort for efficiency. Ideally that state is per-scene renderer but the renderer doesn't know anything about particles.
 					FGPUSortManager::FAllocationInfo SortedIndicesInfo;
-					if (FXSystem->AddSortedGPUSimulation(Simulation, View->ViewMatrices.GetViewOrigin(), bTranslucent, SortedIndicesInfo))
+					if (FXSystem->AddSortedGPUSimulation(RHICmdList, Simulation, View->ViewMatrices.GetViewOrigin(), bTranslucent, SortedIndicesInfo))
 					{
 						MeshBatchUserData = &Collector.AllocateOneFrameResource<FGPUSpriteMeshDataUserData>();
 						MeshBatchUserData->SortedOffset = SortedIndicesInfo.BufferOffset;
@@ -4270,7 +4270,7 @@ void FFXSystem::RemoveGPUSimulation(FParticleSimulationGPU* Simulation)
 	}
 }
 
-bool FFXSystem::AddSortedGPUSimulation(FParticleSimulationGPU* Simulation, const FVector& ViewOrigin, bool bIsTranslucent, FGPUSortManager::FAllocationInfo& OutInfo)
+bool FFXSystem::AddSortedGPUSimulation(FRHICommandListBase& RHICmdList, FParticleSimulationGPU* Simulation, const FVector& ViewOrigin, bool bIsTranslucent, FGPUSortManager::FAllocationInfo& OutInfo)
 {
 	LLM_SCOPE(ELLMTag::Particles);
 
@@ -4281,7 +4281,7 @@ bool FFXSystem::AddSortedGPUSimulation(FParticleSimulationGPU* Simulation, const
 		EGPUSortFlags::SortAfterPostRenderOpaque;
 
 	// Currently opaque materials would need SortAfterPreRender but this is incompatible with KeyGenAfterPostRenderOpaque
-	if (bIsTranslucent && GPUSortManager && GPUSortManager->AddTask(OutInfo, Simulation->VertexBuffer.ParticleCount, SortFlags))
+	if (bIsTranslucent && GPUSortManager && GPUSortManager->AddTask(RHICmdList, OutInfo, Simulation->VertexBuffer.ParticleCount, SortFlags))
 	{
 		SimulationsToSort.Emplace(Simulation->VertexBuffer.VertexBufferSRV, ViewOrigin, (uint32)Simulation->VertexBuffer.ParticleCount, OutInfo);
 		return true;
diff --git a/Engine/Source/Runtime/Engine/Private/Particles/ParticleSystemRender.cpp b/Engine/Source/Runtime/Engine/Private/Particles/ParticleSystemRender.cpp
--- a/Engine/Source/Runtime/Engine/Private/Particles/ParticleSystemRender.cpp
+++ b/Engine/Source/Runtime/Engine/Private/Particles/ParticleSystemRender.cpp
@@ -845,7 +845,7 @@ void FDynamicSpriteEmitterData::GetDynamicMeshElementsEmitter(const FParticleSys
 
 	const auto FeatureLevel = View->GetFeatureLevel();
 
-	FRHICommandListBase& RHICmdList = FRHICommandListImmediate::Get();
+	FRHICommandListBase& RHICmdList = Collector.GetRHICommandList();
 
 	// Sort and generate particles for this view.
 	const FDynamicSpriteEmitterReplayDataBase* SourceData = GetSourceData();
@@ -1318,7 +1318,6 @@ public:
 uint32 FDynamicMeshEmitterData::GetMeshLODIndexFromProxy(const FParticleSystemSceneProxy *InOwnerProxy) const
 {
 	// Determine first available LOD level, top level can be stripped per platform
-	check(IsInRenderingThread());
 	int32 FirstAvailableLOD = StaticMesh->GetRenderData()->CurrentFirstLODIdx;
 	for (; FirstAvailableLOD < StaticMesh->GetRenderData()->LODResources.Num(); FirstAvailableLOD++)
 	{
@@ -1347,7 +1346,7 @@ void FDynamicMeshEmitterData::GetDynamicMeshElementsEmitter(const FParticleSyste
 	{
 		if (Source.EmitterRenderMode == ERM_Normal)
 		{
-			FRHICommandListBase& RHICmdList = FRHICommandListImmediate::Get();
+			FRHICommandListBase& RHICmdList = Collector.GetRHICommandList();
 
 			const auto FeatureLevel = ViewFamily.GetFeatureLevel();
 			const auto ShaderPlatform = GShaderPlatformForFeatureLevel[FeatureLevel];
@@ -1371,10 +1370,10 @@ void FDynamicMeshEmitterData::GetDynamicMeshElementsEmitter(const FParticleSyste
 
 			FDynamicMeshEmitterCollectorResources& CollectorResources = Collector.AllocateOneFrameResource<FDynamicMeshEmitterCollectorResources>(FeatureLevel);
 			FMeshParticleVertexFactory* MeshVertexFactory = &CollectorResources.VertexFactory;
-			SetupVertexFactory(MeshVertexFactory, LODModel, ChosenLODIdx);
+			SetupVertexFactory(RHICmdList, MeshVertexFactory, LODModel, ChosenLODIdx);
 
 			MeshVertexFactory->SetStrides(InstanceVertexStride, bUsesDynamicParameter ? DynamicParameterVertexStride : 0);
-			MeshVertexFactory->InitResource(FRHICommandListImmediate::Get());
+			MeshVertexFactory->InitResource(RHICmdList);
 
 			const FDynamicSpriteEmitterReplayDataBase* SourceData = GetSourceData();
 			FMeshParticleUniformParameters UniformParameters;
@@ -2283,11 +2282,11 @@ void InitMeshParticleVertexFactoryComponents(FMeshParticleVertexFactory* InVerte
 	Data.bInitialized = true;
 }
 
-void FDynamicMeshEmitterData::SetupVertexFactory( FMeshParticleVertexFactory* InVertexFactory, const FStaticMeshLODResources& LODResources, uint32 LODIdx) const
+void FDynamicMeshEmitterData::SetupVertexFactory( FRHICommandListBase& RHICmdList, FMeshParticleVertexFactory* InVertexFactory, const FStaticMeshLODResources& LODResources, uint32 LODIdx) const
 {
 	FMeshParticleVertexFactory::FDataType Data;
 	InitMeshParticleVertexFactoryComponents(InVertexFactory, LODResources, Data);
-	InVertexFactory->SetData(Data);
+	InVertexFactory->SetData(RHICmdList, Data);
 	InVertexFactory->SetLODIdx((uint8)LODIdx);
 }
 
@@ -2409,7 +2408,7 @@ void FDynamicBeam2EmitterData::GetDynamicMeshElementsEmitter(const FParticleSyst
 		return;
 	}
 
-	FRHICommandListBase& RHICmdList = FRHICommandListImmediate::Get();
+	FRHICommandListBase& RHICmdList = Collector.GetRHICommandList();
 
 	FIndexBuffer* IndexBuffer = nullptr;
 	uint32 FirstIndex = 0;
@@ -5171,7 +5170,7 @@ void FDynamicTrailsEmitterData::GetDynamicMeshElementsEmitter(const FParticleSys
 		return;
 	}
 
-	FRHICommandListBase& RHICmdList = FRHICommandListImmediate::Get();
+	FRHICommandListBase& RHICmdList = Collector.GetRHICommandList();
 
 	const bool bIsWireframe = ViewFamily.EngineShowFlags.Wireframe;
 	FIndexBuffer* IndexBuffer = nullptr;
diff --git a/Engine/Source/Runtime/Engine/Private/PhysicsEngine/PhysDrawing.cpp b/Engine/Source/Runtime/Engine/Private/PhysicsEngine/PhysDrawing.cpp
--- a/Engine/Source/Runtime/Engine/Private/PhysicsEngine/PhysDrawing.cpp
+++ b/Engine/Source/Runtime/Engine/Private/PhysicsEngine/PhysDrawing.cpp
@@ -1088,7 +1088,9 @@ void FKAggregateGeom::GetAggGeom(const FTransform& Transform, const FColor Color
 		else
 			SphylElems[i].DrawElemWire(Collector.GetPDI(ViewIndex), ElemTM, Scale3D, Color);
 	}
-	
+
+	FRHICommandListBase& RHICmdList = Collector.GetRHICommandList();
+
 	if(ConvexElems.Num() > 0)
 	{
 		if(bDrawSolid)
@@ -1116,10 +1118,10 @@ void FKAggregateGeom::GetAggGeom(const FTransform& Transform, const FColor Color
 					&& OutVerts.Num() > 0
 					&& ThisGeom.RenderInfo->IndexBuffer->Indices.Num() > 0)
 				{
-					ThisGeom.RenderInfo->IndexBuffer->InitResource(FRHICommandListImmediate::Get());
+					ThisGeom.RenderInfo->IndexBuffer->InitResource(RHICmdList);
 
 					ThisGeom.RenderInfo->CollisionVertexFactory = new FLocalVertexFactory(Collector.GetFeatureLevel(), "FKAggregateGeom");
-					ThisGeom.RenderInfo->VertexBuffers->InitFromDynamicVertex(ThisGeom.RenderInfo->CollisionVertexFactory, OutVerts);
+					ThisGeom.RenderInfo->VertexBuffers->InitFromDynamicVertex(RHICmdList, ThisGeom.RenderInfo->CollisionVertexFactory, OutVerts);
 
 				}
 			}
@@ -1141,7 +1143,7 @@ void FKAggregateGeom::GetAggGeom(const FTransform& Transform, const FColor Color
 				CalcBoxSphereBounds(LocalBounds, FTransform::Identity);
 
 				FDynamicPrimitiveUniformBuffer& DynamicPrimitiveUniformBuffer = Collector.AllocateOneFrameResource<FDynamicPrimitiveUniformBuffer>();
-				DynamicPrimitiveUniformBuffer.Set(LocalToWorld.ToMatrixWithScale(), LocalToWorld.ToMatrixWithScale(), WorldBounds, LocalBounds, true, false, bOutputVelocity);
+				DynamicPrimitiveUniformBuffer.Set(Collector.GetRHICommandList(), LocalToWorld.ToMatrixWithScale(), LocalToWorld.ToMatrixWithScale(), WorldBounds, LocalBounds, true, false, bOutputVelocity);
 				BatchElement.PrimitiveUniformBufferResource = &DynamicPrimitiveUniformBuffer.UniformBuffer;
 
 			 	// previous l2w not used so treat as static
diff --git a/Engine/Source/Runtime/Engine/Private/SkeletalMesh.cpp b/Engine/Source/Runtime/Engine/Private/SkeletalMesh.cpp
--- a/Engine/Source/Runtime/Engine/Private/SkeletalMesh.cpp
+++ b/Engine/Source/Runtime/Engine/Private/SkeletalMesh.cpp
@@ -6054,7 +6054,7 @@ void FSkeletalMeshSceneProxy::GetMeshElementsConditionallySelectable(const TArra
 	{
 		return;
 	}	
-	MeshObject->PreGDMECallback(ViewFamily.Scene->GetGPUSkinCache(), ViewFamily.FrameCounter);
+	MeshObject->PreGDMECallback(Collector.GetRHICommandList(), ViewFamily.Scene->GetGPUSkinCache(), ViewFamily.FrameCounter);
 
 	const FEngineShowFlags& EngineShowFlags = ViewFamily.EngineShowFlags;
 
diff --git a/Engine/Source/Runtime/Engine/Private/SkeletalRenderGPUSkin.cpp b/Engine/Source/Runtime/Engine/Private/SkeletalRenderGPUSkin.cpp
--- a/Engine/Source/Runtime/Engine/Private/SkeletalRenderGPUSkin.cpp
+++ b/Engine/Source/Runtime/Engine/Private/SkeletalRenderGPUSkin.cpp
@@ -547,11 +547,11 @@ void FSkeletalMeshObjectGPUSkin::UpdateDynamicData_RenderThread(FGPUSkinCache* G
 #endif
 }
 
-void FSkeletalMeshObjectGPUSkin::PreGDMECallback(FGPUSkinCache* GPUSkinCache, uint32 FrameNumber)
+void FSkeletalMeshObjectGPUSkin::PreGDMECallback(FRHICommandList& RHICmdList, FGPUSkinCache* GPUSkinCache, uint32 FrameNumber)
 {
 	if (bNeedsUpdateDeferred)
 	{
-		ProcessUpdatedDynamicData(EGPUSkinCacheEntryMode::Raster, GPUSkinCache, FRHICommandListExecutor::GetImmediateCommandList(), FrameNumber, LastBoneTransformRevisionNumber, bMorphNeedsUpdateDeferred, DynamicData->LODIndex);
+		ProcessUpdatedDynamicData(EGPUSkinCacheEntryMode::Raster, GPUSkinCache, RHICmdList, FrameNumber, LastBoneTransformRevisionNumber, bMorphNeedsUpdateDeferred, DynamicData->LODIndex);
 	}
 }
 
diff --git a/Engine/Source/Runtime/Engine/Private/SkeletalRenderGPUSkin.h b/Engine/Source/Runtime/Engine/Private/SkeletalRenderGPUSkin.h
--- a/Engine/Source/Runtime/Engine/Private/SkeletalRenderGPUSkin.h
+++ b/Engine/Source/Runtime/Engine/Private/SkeletalRenderGPUSkin.h
@@ -353,7 +353,7 @@ public:
 	ENGINE_API virtual void ReleaseResources() override;
 	ENGINE_API virtual void Update(int32 LODIndex,USkinnedMeshComponent* InMeshComponent,const FMorphTargetWeightMap& InActiveMorphTargets, const TArray<float>& InMorphTargetWeights, EPreviousBoneTransformUpdateMode PreviousBoneTransformUpdateMode, const FExternalMorphWeightData& InExternalMorphWeightData) override;
 	ENGINE_API void UpdateDynamicData_RenderThread(FGPUSkinCache* GPUSkinCache, FRHICommandList& RHICmdList, FDynamicSkelMeshObjectDataGPUSkin* InDynamicData, FSceneInterface* Scene, uint64 FrameNumberToPrepare, uint32 RevisionNumber);
-	ENGINE_API virtual void PreGDMECallback(FGPUSkinCache* GPUSkinCache, uint32 FrameNumber) override;
+	ENGINE_API virtual void PreGDMECallback(FRHICommandList& RHICmdList, FGPUSkinCache* GPUSkinCache, uint32 FrameNumber) override;
 	ENGINE_API virtual const FVertexFactory* GetSkinVertexFactory(const FSceneView* View, int32 LODIndex,int32 ChunkIdx, ESkinVertexFactoryMode VFMode = ESkinVertexFactoryMode::Default) const override;
 	ENGINE_API virtual const FSkinBatchVertexFactoryUserData* GetVertexFactoryUserData(const int32 LODIndex, int32 ChunkIdx, ESkinVertexFactoryMode VFMode) const override;
 	virtual void CacheVertices(int32 LODIndex, bool bForce) const override {}
diff --git a/Engine/Source/Runtime/Engine/Private/VectorField.cpp b/Engine/Source/Runtime/Engine/Private/VectorField.cpp
--- a/Engine/Source/Runtime/Engine/Private/VectorField.cpp
+++ b/Engine/Source/Runtime/Engine/Private/VectorField.cpp
@@ -748,7 +748,7 @@ public:
 				if (IsSelected() || View->Family->EngineShowFlags.VectorFields)
 				{
 					FVectorFieldCollectorResources& CollectorResources = Collector.AllocateOneFrameResource<FVectorFieldCollectorResources>(View->GetFeatureLevel());
-					CollectorResources.VisualizationVertexFactory.InitResource(FRHICommandListImmediate::Get());
+					CollectorResources.VisualizationVertexFactory.InitResource(Collector.GetRHICommandList());
 
 					GetVectorFieldMesh(&CollectorResources.VisualizationVertexFactory, VectorFieldInstance, ViewIndex, Collector);
 				}
diff --git a/Engine/Source/Runtime/Engine/Public/GPUSortManager.h b/Engine/Source/Runtime/Engine/Public/GPUSortManager.h
--- a/Engine/Source/Runtime/Engine/Public/GPUSortManager.h
+++ b/Engine/Source/Runtime/Engine/Public/GPUSortManager.h
@@ -422,8 +422,10 @@ public:
 	 * 
 	 * Return true if the task was added to a batch, false otherwise, could be because GPU sorting is disabled.
 	 */
-	ENGINE_API bool AddTask(FAllocationInfo& OutInfo, int32 ValueCount, EGPUSortFlags TaskFlags);
+	ENGINE_API bool AddTask(FRHICommandListBase& RHICmdList, FAllocationInfo& OutInfo, int32 ValueCount, EGPUSortFlags TaskFlags);
 
+	UE_DEPRECATED(5.4, "AddTask now requires a command list.")
+	ENGINE_API bool AddTask(FAllocationInfo& OutInfo, int32 ValueCount, EGPUSortFlags TaskFlags);
 
 	/**
 	 * Callback that needs to be called in the rendering loop, after calls to FFXSystemInterface::PreRender() are issued.
diff --git a/Engine/Source/Runtime/Engine/Public/MeshParticleVertexFactory.h b/Engine/Source/Runtime/Engine/Public/MeshParticleVertexFactory.h
--- a/Engine/Source/Runtime/Engine/Public/MeshParticleVertexFactory.h
+++ b/Engine/Source/Runtime/Engine/Public/MeshParticleVertexFactory.h
@@ -123,6 +123,9 @@ public:
 	/**
 	 * An implementation of the interface used by TSynchronizedResource to update the resource with new data from the game thread.
 	 */
+	ENGINE_API void SetData(FRHICommandListBase& RHICmdList, const FDataType& InData);
+
+	UE_DEPRECATED(5.4, "SetData requires a command list.")
 	ENGINE_API void SetData(const FDataType& InData);
 
 	/**
diff --git a/Engine/Source/Runtime/Engine/Public/ParticleHelper.h b/Engine/Source/Runtime/Engine/Public/ParticleHelper.h
--- a/Engine/Source/Runtime/Engine/Public/ParticleHelper.h
+++ b/Engine/Source/Runtime/Engine/Public/ParticleHelper.h
@@ -1827,7 +1827,7 @@ struct FDynamicMeshEmitterData : public FDynamicSpriteEmitterDataBase
 	/**
 	 *	 Initialize this emitter's vertex factory with the vertex buffers from the mesh's rendering data.
 	 */
-	void SetupVertexFactory( FMeshParticleVertexFactory* InVertexFactory, const FStaticMeshLODResources& LODResources, uint32 LODIdx) const;
+	void SetupVertexFactory( FRHICommandListBase& RHICmdList, FMeshParticleVertexFactory* InVertexFactory, const FStaticMeshLODResources& LODResources, uint32 LODIdx) const;
 
 	/** Returns the source data for this particle system */
 	virtual const FDynamicEmitterReplayDataBase& GetSource() const override
diff --git a/Engine/Source/Runtime/Engine/Public/SceneManagement.h b/Engine/Source/Runtime/Engine/Public/SceneManagement.h
--- a/Engine/Source/Runtime/Engine/Public/SceneManagement.h
+++ b/Engine/Source/Runtime/Engine/Public/SceneManagement.h
@@ -2188,7 +2188,7 @@ public:
 		bool bHasPrecomputedVolumetricLightmap,
 		bool bOutputVelocity);
 
-	//UE_DEPRECATED(5.4, "Set requires a command list")
+	UE_DEPRECATED(5.4, "Set requires a command list")
 	ENGINE_API void Set(
 		const FMatrix& LocalToWorld,
 		const FMatrix& PreviousLocalToWorld,
@@ -2201,7 +2201,7 @@ public:
 		bool bOutputVelocity,
 		const FCustomPrimitiveData* CustomPrimitiveData);
 
-	//UE_DEPRECATED(5.4, "Set requires a command list")
+	UE_DEPRECATED(5.4, "Set requires a command list")
 	ENGINE_API void Set(
 		const FMatrix& LocalToWorld,
 		const FMatrix& PreviousLocalToWorld,
@@ -2213,7 +2213,7 @@ public:
 		bool bOutputVelocity,
 		const FCustomPrimitiveData* CustomPrimitiveData);
 
-	//UE_DEPRECATED(5.4, "Set requires a command list")
+	UE_DEPRECATED(5.4, "Set requires a command list")
 	ENGINE_API void Set(
 		const FMatrix& LocalToWorld,
 		const FMatrix& PreviousLocalToWorld,
@@ -2225,7 +2225,7 @@ public:
 		bool bOutputVelocity);
 
 	/** Pass-through implementation which calls the overloaded Set function with LocalBounds for PreSkinnedLocalBounds. */
-	//UE_DEPRECATED(5.4, "Set requires a command list")
+	UE_DEPRECATED(5.4, "Set requires a command list")
 	ENGINE_API void Set(
 		const FMatrix& LocalToWorld,
 		const FMatrix& PreviousLocalToWorld,
diff --git a/Engine/Source/Runtime/Engine/Public/SkeletalRenderPublic.h b/Engine/Source/Runtime/Engine/Public/SkeletalRenderPublic.h
--- a/Engine/Source/Runtime/Engine/Public/SkeletalRenderPublic.h
+++ b/Engine/Source/Runtime/Engine/Public/SkeletalRenderPublic.h
@@ -29,6 +29,7 @@ struct FRWBuffer;
 class FGPUSkinCacheEntry;
 class FMeshDeformerGeometry;
 class FRayTracingGeometry;
+class FRHICommandList;
 
 /** data for a single skinned skeletal mesh vertex */
 struct FFinalSkinVertex
@@ -93,7 +94,7 @@ public:
 	 * Called by FSkeletalMeshObject prior to GDME. This allows the GPU skin version to update bones etc now that we know we are going to render
 	 * @param FrameNumber from GFrameNumber
 	 */
-	virtual void PreGDMECallback(class FGPUSkinCache* GPUSkinCache, uint32 FrameNumber)
+	virtual void PreGDMECallback(FRHICommandList& RHICmdList, class FGPUSkinCache* GPUSkinCache, uint32 FrameNumber)
 	{
 	}
 
diff --git a/Engine/Source/Runtime/GeometryFramework/Private/Components/BaseDynamicMeshSceneProxy.cpp b/Engine/Source/Runtime/GeometryFramework/Private/Components/BaseDynamicMeshSceneProxy.cpp
--- a/Engine/Source/Runtime/GeometryFramework/Private/Components/BaseDynamicMeshSceneProxy.cpp
+++ b/Engine/Source/Runtime/GeometryFramework/Private/Components/BaseDynamicMeshSceneProxy.cpp
@@ -277,7 +277,7 @@ void FBaseDynamicMeshSceneProxy::GetDynamicMeshElements(const TArray<const FScen
 
 				// do we need separate one of these for each MeshRenderBufferSet?
 				FDynamicPrimitiveUniformBuffer& DynamicPrimitiveUniformBuffer = Collector.AllocateOneFrameResource<FDynamicPrimitiveUniformBuffer>();
-				DynamicPrimitiveUniformBuffer.Set(GetLocalToWorld(), PreviousLocalToWorld, GetBounds(), GetLocalBounds(), GetLocalBounds(), true, bHasPrecomputedVolumetricLightmap, bOutputVelocity, GetCustomPrimitiveData());
+				DynamicPrimitiveUniformBuffer.Set(Collector.GetRHICommandList(), GetLocalToWorld(), PreviousLocalToWorld, GetBounds(), GetLocalBounds(), GetLocalBounds(), true, bHasPrecomputedVolumetricLightmap, bOutputVelocity, GetCustomPrimitiveData());
 
 				// If we want Wireframe-on-Shaded, we have to draw the solid. If View Mode Overrides are enabled, the solid
 				// will be replaced with it's wireframe, so we might as well not. 
@@ -531,7 +531,7 @@ void FBaseDynamicMeshSceneProxy::GetDynamicRayTracingInstances(FRayTracingMateri
 
 	// is it safe to share this between primary and secondary raytracing batches?
 	FDynamicPrimitiveUniformBuffer& DynamicPrimitiveUniformBuffer = Context.RayTracingMeshResourceCollector.AllocateOneFrameResource<FDynamicPrimitiveUniformBuffer>();
-	DynamicPrimitiveUniformBuffer.Set(GetLocalToWorld(), PreviousLocalToWorld, GetBounds(), GetLocalBounds(), true, bHasPrecomputedVolumetricLightmap, bOutputVelocity);
+	DynamicPrimitiveUniformBuffer.Set(Context.RHICmdList, GetLocalToWorld(), PreviousLocalToWorld, GetBounds(), GetLocalBounds(), true, bHasPrecomputedVolumetricLightmap, bOutputVelocity);
 
 	// Draw the active buffer sets
 	for (FMeshRenderBufferSet* BufferSet : Buffers)
diff --git a/Engine/Source/Runtime/Landscape/Private/LandscapeCollision.cpp b/Engine/Source/Runtime/Landscape/Private/LandscapeCollision.cpp
--- a/Engine/Source/Runtime/Landscape/Private/LandscapeCollision.cpp
+++ b/Engine/Source/Runtime/Landscape/Private/LandscapeCollision.cpp
@@ -717,7 +717,7 @@ FPrimitiveSceneProxy* ULandscapeHeightfieldCollisionComponent::CreateSceneProxy(
 						BatchElement.MaxVertexIndex = VertexBuffers.PositionVertexBuffer.GetNumVertices() - 1;
 
 						FDynamicPrimitiveUniformBuffer& DynamicPrimitiveUniformBuffer = Collector.AllocateOneFrameResource<FDynamicPrimitiveUniformBuffer>();
-						DynamicPrimitiveUniformBuffer.Set(LocalToWorldNoScale, LocalToWorldNoScale, Bounds, GetLocalBounds(), /*bReceivesDecals = */false, /*bHasPrecomputedVolumetricLightmap = */false, AlwaysHasVelocity());
+						DynamicPrimitiveUniformBuffer.Set(Collector.GetRHICommandList(), LocalToWorldNoScale, LocalToWorldNoScale, Bounds, GetLocalBounds(), /*bReceivesDecals = */false, /*bHasPrecomputedVolumetricLightmap = */false, AlwaysHasVelocity());
 						BatchElement.PrimitiveUniformBufferResource = &DynamicPrimitiveUniformBuffer.UniformBuffer;
 
 						Collector.AddMesh(ViewIndex, MeshBatch);
diff --git a/Engine/Source/Runtime/MRMesh/Private/MRMeshComponent.cpp b/Engine/Source/Runtime/MRMesh/Private/MRMeshComponent.cpp
--- a/Engine/Source/Runtime/MRMesh/Private/MRMeshComponent.cpp
+++ b/Engine/Source/Runtime/MRMesh/Private/MRMeshComponent.cpp
@@ -440,7 +440,7 @@ private:
 							Mesh.MaterialRenderProxy = MaterialProxy;
 
 							FDynamicPrimitiveUniformBuffer& DynamicPrimitiveUniformBuffer = Collector.AllocateOneFrameResource<FDynamicPrimitiveUniformBuffer>();
-							DynamicPrimitiveUniformBuffer.Set(GetLocalToWorld(), GetLocalToWorld(), InfiniteBounds, InfiniteBounds, true, false, AlwaysHasVelocity());
+							DynamicPrimitiveUniformBuffer.Set(Collector.GetRHICommandList(), GetLocalToWorld(), GetLocalToWorld(), InfiniteBounds, InfiniteBounds, true, false, AlwaysHasVelocity());
 							BatchElement.PrimitiveUniformBufferResource = &DynamicPrimitiveUniformBuffer.UniformBuffer;
 
 							BatchElement.FirstIndex = 0;
diff --git a/Engine/Source/Runtime/NavigationSystem/Private/NavMesh/NavMeshRenderingComponent.cpp b/Engine/Source/Runtime/NavigationSystem/Private/NavMesh/NavMeshRenderingComponent.cpp
--- a/Engine/Source/Runtime/NavigationSystem/Private/NavMesh/NavMeshRenderingComponent.cpp
+++ b/Engine/Source/Runtime/NavigationSystem/Private/NavMesh/NavMeshRenderingComponent.cpp
@@ -1477,7 +1477,7 @@ void FNavMeshSceneProxy::GetDynamicMeshElements(const TArray<const FSceneView*>&
 					BatchElement = MeshBatchElements[Index];
 
 					FDynamicPrimitiveUniformBuffer& DynamicPrimitiveUniformBuffer = Collector.AllocateOneFrameResource<FDynamicPrimitiveUniformBuffer>();
-					DynamicPrimitiveUniformBuffer.Set(FMatrix::Identity, FMatrix::Identity, GetBounds(), GetLocalBounds(), false, false, AlwaysHasVelocity());
+					DynamicPrimitiveUniformBuffer.Set(Collector.GetRHICommandList(), FMatrix::Identity, FMatrix::Identity, GetBounds(), GetLocalBounds(), false, false, AlwaysHasVelocity());
 					BatchElement.PrimitiveUniformBufferResource = &DynamicPrimitiveUniformBuffer.UniformBuffer;
 
 					Mesh.bWireframe = false;
diff --git a/Engine/Source/Runtime/NavigationSystem/Private/NavMesh/NavTestRenderingComponent.cpp b/Engine/Source/Runtime/NavigationSystem/Private/NavMesh/NavTestRenderingComponent.cpp
--- a/Engine/Source/Runtime/NavigationSystem/Private/NavMesh/NavTestRenderingComponent.cpp
+++ b/Engine/Source/Runtime/NavigationSystem/Private/NavMesh/NavTestRenderingComponent.cpp
@@ -53,6 +53,8 @@ FNavTestSceneProxy::FNavTestSceneProxy(const UNavTestRenderingComponent* InCompo
 
 void FNavTestSceneProxy::GetDynamicMeshElements(const TArray<const FSceneView*>& Views, const FSceneViewFamily& ViewFamily, uint32 VisibilityMap, FMeshElementCollector& Collector) const
 {
+	FRHICommandList& RHICmdList = Collector.GetRHICommandList();
+
 	for (int32 ViewIndex = 0; ViewIndex < Views.Num(); ViewIndex++)
 	{
 		if (VisibilityMap & (1 << ViewIndex))
diff --git a/Engine/Source/Runtime/Renderer/Private/SparseVolumeTexture/SparseVolumeTextureViewerSceneProxy.cpp b/Engine/Source/Runtime/Renderer/Private/SparseVolumeTexture/SparseVolumeTextureViewerSceneProxy.cpp
--- a/Engine/Source/Runtime/Renderer/Private/SparseVolumeTexture/SparseVolumeTextureViewerSceneProxy.cpp
+++ b/Engine/Source/Runtime/Renderer/Private/SparseVolumeTexture/SparseVolumeTextureViewerSceneProxy.cpp
@@ -54,8 +54,6 @@ SIZE_T FSparseVolumeTextureViewerSceneProxy::GetTypeHash() const
 
 void FSparseVolumeTextureViewerSceneProxy::GetDynamicMeshElements(const TArray<const FSceneView*>& Views, const FSceneViewFamily& ViewFamily, uint32 VisibilityMap, FMeshElementCollector& Collector) const
 {
-	check(IsInRenderingThread());
-
 	for (int32 ViewIndex = 0; ViewIndex < Views.Num(); ViewIndex++)
 	{
 		const FSceneView* View = Views[ViewIndex];
