VHM: Set explicit wavesize as temporary fix for hardware with wavesize 16. Suspect that this is just a workaround for bug in this shader's persistant thread implemention.
#rnx

diff --git a/Engine/Plugins/Experimental/VirtualHeightfieldMesh/Shaders/Private/VirtualHeightfieldMesh.usf b/Engine/Plugins/Experimental/VirtualHeightfieldMesh/Shaders/Private/VirtualHeightfieldMesh.usf
--- a/Engine/Plugins/Experimental/VirtualHeightfieldMesh/Shaders/Private/VirtualHeightfieldMesh.usf
+++ b/Engine/Plugins/Experimental/VirtualHeightfieldMesh/Shaders/Private/VirtualHeightfieldMesh.usf
@@ -212,6 +212,9 @@ void InitBuffersCS()
  * Compute shader to traverse the virtual texture page table and generate an array of items to potentially render for a view.
  */
 
+#if COMPILER_SUPPORTS_WAVE_SIZE
+	WAVESIZE(32)
+#endif
 [numthreads(64, 1, 1)]
 void CollectQuadsCS(
 	uint3 DispatchThreadId : SV_DispatchThreadID,
