Use recursive mutex to fix deadlock in virtual heightfield renderer extension.

diff --git a/Engine/Plugins/Experimental/VirtualHeightfieldMesh/Source/VirtualHeightfieldMesh/Private/VirtualHeightfieldMeshSceneProxy.cpp b/Engine/Plugins/Experimental/VirtualHeightfieldMesh/Source/VirtualHeightfieldMesh/Private/VirtualHeightfieldMeshSceneProxy.cpp
--- a/Engine/Plugins/Experimental/VirtualHeightfieldMesh/Source/VirtualHeightfieldMesh/Private/VirtualHeightfieldMeshSceneProxy.cpp
+++ b/Engine/Plugins/Experimental/VirtualHeightfieldMesh/Source/VirtualHeightfieldMesh/Private/VirtualHeightfieldMeshSceneProxy.cpp
@@ -203,7 +203,7 @@ private:
 	/** Flag for frame validation. */
 	bool bInFrame;
 
-	UE::FMutex Mutex;
+	UE::FRecursiveMutex Mutex;
 
 	/** Buffers to fill. Resources can persist between frames to reduce allocation cost, but contents don't persist. */
 	TArray<VirtualHeightfieldMesh::FDrawInstanceBuffers> Buffers;
