PR #11142: VirtualHeightfieldMesh: fix crash caused by incorrect mipmap setup

diff --git a/Engine/Plugins/Experimental/VirtualHeightfieldMesh/Source/VirtualHeightfieldMesh/Private/HeightfieldMinMaxTexture.cpp b/Engine/Plugins/Experimental/VirtualHeightfieldMesh/Source/VirtualHeightfieldMesh/Private/HeightfieldMinMaxTexture.cpp
--- a/Engine/Plugins/Experimental/VirtualHeightfieldMesh/Source/VirtualHeightfieldMesh/Private/HeightfieldMinMaxTexture.cpp
+++ b/Engine/Plugins/Experimental/VirtualHeightfieldMesh/Source/VirtualHeightfieldMesh/Private/HeightfieldMinMaxTexture.cpp
@@ -152,7 +152,7 @@ void UHeightfieldMinMaxTexture::RebuildLodBiasTexture(FHeightfieldMinMaxTextureB
 		LodBiasTexture->Filter = TF_Nearest;
 		LodBiasTexture->MipGenSettings = TMGS_NoMipmaps;
 		LodBiasTexture->SetLayerFormatSettings(0, Settings);
-		LodBiasTexture->Source.Init(InBuildDesc.SizeX, InBuildDesc.SizeY, 1, InBuildDesc.NumMips, TSF_G8, LodBiasTextureData.GetData());
+		LodBiasTexture->Source.Init(InBuildDesc.SizeX, InBuildDesc.SizeY, 1, 1, TSF_G8, LodBiasTextureData.GetData());
 
 		LodBiasTexture->PostEditChange();
 	}
