Remove use of soon to be deprecated function.
#rnx

diff --git a/Engine/Plugins/Experimental/VirtualHeightfieldMesh/Source/VirtualHeightfieldMesh/Private/VirtualHeightfieldMeshSceneProxy.cpp b/Engine/Plugins/Experimental/VirtualHeightfieldMesh/Source/VirtualHeightfieldMesh/Private/VirtualHeightfieldMeshSceneProxy.cpp
--- a/Engine/Plugins/Experimental/VirtualHeightfieldMesh/Source/VirtualHeightfieldMesh/Private/VirtualHeightfieldMeshSceneProxy.cpp
+++ b/Engine/Plugins/Experimental/VirtualHeightfieldMesh/Source/VirtualHeightfieldMesh/Private/VirtualHeightfieldMeshSceneProxy.cpp
@@ -1032,7 +1032,7 @@ namespace VirtualHeightfieldMesh
 		OutResources.QuadBufferUAV = GraphBuilder.CreateUAV(FRDGBufferUAVDesc(OutResources.QuadBuffer, PF_R32G32_UINT));
 		OutResources.QuadBufferSRV = GraphBuilder.CreateSRV(FRDGBufferSRVDesc(OutResources.QuadBuffer, PF_R32G32_UINT));
 
-		FRDGBufferDesc FeedbackBufferDesc = FRDGBufferDesc::CreateBufferDesc(sizeof(uint32), InDesc.MaxFeedbackItems + 1);
+		FRDGBufferDesc FeedbackBufferDesc = FRDGBufferDesc::CreateBufferDesc(sizeof(uint32), InDesc.MaxFeedbackItems);
 		FeedbackBufferDesc.Usage = EBufferUsageFlags(FeedbackBufferDesc.Usage | BUF_SourceCopy);
 		OutResources.FeedbackBuffer = GraphBuilder.CreateBuffer(FeedbackBufferDesc, TEXT("VirtualHeightfieldMesh.FeedbackBuffer"));
 		OutResources.FeedbackBufferUAV = GraphBuilder.CreateUAV(FRDGBufferUAVDesc(OutResources.FeedbackBuffer, PF_R32_UINT));
@@ -1341,7 +1341,8 @@ void FVirtualHeightfieldMeshRendererExtension::SubmitWork(FRDGBuilder& GraphBuil
 			VirtualHeightfieldMesh::AddPass_CollectQuads(GraphBuilder, GetGlobalShaderMap(GMaxRHIFeatureLevel), ProxyDesc, VolatileResources, MainViewDesc);
 
 			FVirtualTextureFeedbackBufferDesc Desc;
- 			Desc.Init(CVarVHMMaxFeedbackItems.GetValueOnRenderThread() + 1);
+ 			Desc.BufferSize = CVarVHMMaxFeedbackItems.GetValueOnRenderThread();
+			Desc.bSizeInHeader = true;
  			SubmitVirtualTextureFeedbackBuffer(GraphBuilder, VolatileResources.FeedbackBuffer, Desc);
 
 			while (WorkIndex < NumWorkItems && SceneProxies[WorkDescs[WorkIndex].ProxyIndex] == Proxy && MainViews[WorkDescs[WorkIndex].MainViewIndex] == MainView)
